Fixed no-unity compile errors.

diff --git a/Engine/Source/Programs/UnrealInsights/Private/Windows/WindowsUnrealInsightsMain.cpp b/Engine/Source/Programs/UnrealInsights/Private/Windows/WindowsUnrealInsightsMain.cpp
--- a/Engine/Source/Programs/UnrealInsights/Private/Windows/WindowsUnrealInsightsMain.cpp
+++ b/Engine/Source/Programs/UnrealInsights/Private/Windows/WindowsUnrealInsightsMain.cpp
@@ -3,6 +3,8 @@
 #include "UnrealInsightsMain.h"
 #include "HAL/ExceptionHandling.h"
 #include "Windows/WindowsHWrapper.h"
+#include "LaunchEngineLoop.h"
+#include "Misc/CommandLine.h"
 #include "Misc/OutputDeviceError.h"
 
 /**
