Removed unused module dependency.

diff --git a/Engine/Source/Programs/UnrealInsights/UnrealInsights.Build.cs b/Engine/Source/Programs/UnrealInsights/UnrealInsights.Build.cs
--- a/Engine/Source/Programs/UnrealInsights/UnrealInsights.Build.cs
+++ b/Engine/Source/Programs/UnrealInsights/UnrealInsights.Build.cs
@@ -21,7 +21,6 @@ public class UnrealInsights : ModuleRules
 				"SlateCore",
 				"SourceCodeAccess",
 				"StandaloneRenderer",
-				"Trace",
 				"TraceInsights",
 				"TraceServices"
 			}
