Fixed Mac compile error.

diff --git a/Engine/Source/Programs/UnrealInsights/Private/Mac/MacUnrealInsightsMain.cpp b/Engine/Source/Programs/UnrealInsights/Private/Mac/MacUnrealInsightsMain.cpp
--- a/Engine/Source/Programs/UnrealInsights/Private/Mac/MacUnrealInsightsMain.cpp
+++ b/Engine/Source/Programs/UnrealInsights/Private/Mac/MacUnrealInsightsMain.cpp
@@ -2,6 +2,7 @@
 
 #include "UnrealInsightsMain.h"
 #include "HAL/ExceptionHandling.h"
+#include "LaunchEngineLoop.h"
 #include "Mac/CocoaThread.h"
 
 
