Removed IOS decoration which isn't appropriate for UnrealInsights (a tool).

diff --git a/Engine/Source/Programs/UnrealInsights/UnrealInsights.Target.cs b/Engine/Source/Programs/UnrealInsights/UnrealInsights.Target.cs
--- a/Engine/Source/Programs/UnrealInsights/UnrealInsights.Target.cs
+++ b/Engine/Source/Programs/UnrealInsights/UnrealInsights.Target.cs
@@ -4,7 +4,6 @@ using UnrealBuildTool;
 using System.Collections.Generic;
 
 [SupportedPlatforms(UnrealPlatformClass.Desktop)]
-[SupportedPlatforms(UnrealTargetPlatform.IOS)]
 public class UnrealInsightsTarget : TargetRules
 {
 	public UnrealInsightsTarget(TargetInfo Target) : base(Target)
