Only build UnrealInsights for Win64 for now.

diff --git a/Engine/Source/Programs/UnrealInsights/UnrealInsights.Target.cs b/Engine/Source/Programs/UnrealInsights/UnrealInsights.Target.cs
--- a/Engine/Source/Programs/UnrealInsights/UnrealInsights.Target.cs
+++ b/Engine/Source/Programs/UnrealInsights/UnrealInsights.Target.cs
@@ -3,7 +3,7 @@
 using UnrealBuildTool;
 using System.Collections.Generic;
 
-[SupportedPlatforms(UnrealPlatformClass.Desktop)]
+[SupportedPlatforms(UnrealTargetPlatform.Win64)]
 public class UnrealInsightsTarget : TargetRules
 {
 	public UnrealInsightsTarget(TargetInfo Target) : base(Target)
