Fixed fallout from platform module changes

diff --git a/Engine/Source/Programs/UnrealInsights/UnrealInsights.Target.cs b/Engine/Source/Programs/UnrealInsights/UnrealInsights.Target.cs
--- a/Engine/Source/Programs/UnrealInsights/UnrealInsights.Target.cs
+++ b/Engine/Source/Programs/UnrealInsights/UnrealInsights.Target.cs
@@ -3,7 +3,7 @@
 using UnrealBuildTool;
 using System.Collections.Generic;
 
-[SupportedPlatforms(UnrealTargetPlatform.Win64)]
+[SupportedPlatforms("Win64")]
 public class UnrealInsightsTarget : TargetRules
 {
 	public UnrealInsightsTarget(TargetInfo Target) : base(Target)
