[Insights] Disabled menu animations as these aren't supported. Fixes UE-75583 (Timing Insights / Category Filter flickers and entries can't be selected when using search feature).

diff --git a/Engine/Source/Programs/UnrealInsights/Private/UserInterfaceCommand.cpp b/Engine/Source/Programs/UnrealInsights/Private/UserInterfaceCommand.cpp
--- a/Engine/Source/Programs/UnrealInsights/Private/UserInterfaceCommand.cpp
+++ b/Engine/Source/Programs/UnrealInsights/Private/UserInterfaceCommand.cpp
@@ -110,9 +110,14 @@ void FUserInterfaceCommand::Run()
 
 void FUserInterfaceCommand::InitializeSlateApplication(const FString& LayoutIni)
 {
+	//TODO: FSlateApplication::InitHighDPI(true);
+
 	// Crank up a normal Slate application using the platform's standalone renderer.
 	FSlateApplication::InitializeAsStandaloneApplication(GetStandardStandaloneRenderer());
 
+	// Menu anims aren't supported. See Runtime\Slate\Private\Framework\Application\MenuStack.cpp.
+	FSlateApplication::Get().EnableMenuAnimations(false);
+
 	// Set the application name.
 	FGlobalTabmanager::Get()->SetApplicationTitle(NSLOCTEXT("UnrealInsights", "AppTitle", "Unreal Insights"));
 
