[Insights] Added code to allow only a single instance of the Session Browser process.

diff --git a/Engine/Source/Programs/UnrealInsights/Private/UserInterfaceCommand.cpp b/Engine/Source/Programs/UnrealInsights/Private/UserInterfaceCommand.cpp
--- a/Engine/Source/Programs/UnrealInsights/Private/UserInterfaceCommand.cpp
+++ b/Engine/Source/Programs/UnrealInsights/Private/UserInterfaceCommand.cpp
@@ -25,6 +25,16 @@
 #include "Insights/IUnrealInsightsModule.h"
 #include "Insights/Version.h"
 
+#if PLATFORM_WINDOWS
+#include "Windows/AllowWindowsPlatformTypes.h"
+#include "Windows/MinWindows.h"
+#include "Windows/HideWindowsPlatformTypes.h"
+#endif
+#if PLATFORM_UNIX
+#include <sys/file.h>
+#include <errno.h>
+#endif
+
 ////////////////////////////////////////////////////////////////////////////////////////////////////
 
 #define IDEAL_FRAMERATE 60
@@ -38,8 +48,68 @@ namespace UserInterfaceCommand
 
 ////////////////////////////////////////////////////////////////////////////////////////////////////
 
+bool CheckSessionBrowserSingleInstance()
+{
+#if PLATFORM_WINDOWS
+	HANDLE SessionBrowserEvent = CreateEventW(NULL, 0, 0, L"Global\\UnrealInsights.SessionBrowser");
+	if (SessionBrowserEvent == NULL || GetLastError() == ERROR_ALREADY_EXISTS)
+	{
+		// Another Session Browser process is already running.
+
+		if (SessionBrowserEvent != NULL)
+		{
+			CloseHandle(SessionBrowserEvent);
+		}
+
+		// Activate the respective window.
+		HWND Window = FindWindowW(0, L"Unreal Insights");
+		if (Window)
+		{
+			ShowWindow(Window, SW_SHOW);
+			SetForegroundWindow(Window);
+		}
+
+		return false;
+	}
+#endif // PLATFORM_WINDOWS
+
+#if PLATFORM_UNIX
+	int FileHandle = open("/var/run/UnrealInsights.pid", O_CREAT | O_RDWR, 0666);
+	int Ret = flock(FileHandle, LOCK_EX | LOCK_NB);
+	if (Ret && EWOULDBLOCK == errno)
+	{
+		// Another Session Browser process is already running.
+
+		// Activate the respective window.
+		//TODO: "wmctrl -a Insights"
+
+		return false;
+	}
+#endif
+
+	return true;
+}
+
+////////////////////////////////////////////////////////////////////////////////////////////////////
+
 void FUserInterfaceCommand::Run()
 {
+	// Only a single instance of Session Browser window/process is allowed.
+	{
+		const uint32 MaxPath = FPlatformMisc::GetMaxPathLength();
+		TCHAR* TraceFile = new TCHAR[MaxPath + 1];
+		TraceFile[0] = 0;
+		bool bUseTraceFile = FParse::Value(FCommandLine::Get(), TEXT("-Trace="), TraceFile, MaxPath, true);
+
+		if (!bUseTraceFile)
+		{
+			if (!CheckSessionBrowserSingleInstance())
+			{
+				return;
+			}
+		}
+	}
+
 	FCoreStyle::ResetToDefault();
 
 	// Load required modules.
