Restored auto-tracing on Windows if Insights is running.

#rnx

diff --git a/Engine/Source/Developer/TraceServices/Private/SessionService.cpp b/Engine/Source/Developer/TraceServices/Private/SessionService.cpp
--- a/Engine/Source/Developer/TraceServices/Private/SessionService.cpp
+++ b/Engine/Source/Developer/TraceServices/Private/SessionService.cpp
@@ -14,12 +14,6 @@
 #include "IPAddress.h"
 #include "AddressInfoTypes.h"
 
-#if PLATFORM_WINDOWS
-	#include "Windows/AllowWindowsPlatformTypes.h"
-	#include "Windows/MinWindows.h"
-	#include "Windows/HideWindowsPlatformTypes.h"
-#endif
-
 namespace Trace
 {
 
@@ -108,14 +102,6 @@ FSessionService::~FSessionService()
 bool FSessionService::StartRecorderServer()
 {
 	bool bOk = TraceRecorder->StartRecording();
-#if PLATFORM_WINDOWS
-	// Create a named event that other processes can use detect a running
-	// recorder and connect to it automatically
-	if (bOk && RecorderEvent == nullptr)
-	{
-		RecorderEvent = ::CreateEvent(nullptr, true, false, TEXT("Local\\UnrealInsightsRecorder"));
-	}
-#endif // PLATFORM_WINDOWS
 	return bOk;
 }
 
@@ -126,13 +112,6 @@ bool FSessionService::IsRecorderServerRunning() const
 
 void FSessionService::StopRecorderServer()
 {
-#if PLATFORM_WINDOWS
-	if (RecorderEvent != nullptr)
-	{
-		::CloseHandle(RecorderEvent);
-		RecorderEvent = nullptr;
-	}
-#endif // PLATFORM_WINDOWS
 	TraceRecorder->StopRecording();
 }
 
diff --git a/Engine/Source/Developer/TraceServices/Private/SessionServicePrivate.h b/Engine/Source/Developer/TraceServices/Private/SessionServicePrivate.h
--- a/Engine/Source/Developer/TraceServices/Private/SessionServicePrivate.h
+++ b/Engine/Source/Developer/TraceServices/Private/SessionServicePrivate.h
@@ -57,7 +57,6 @@ private:
 
 	FModuleService& ModuleService;
 	FAnalysisService& AnalysisService;
-	void* RecorderEvent = nullptr;
 	FString LocalSessionDirectory;
 	TSharedPtr<IStore> TraceStore;
 	TSharedPtr<IRecorder> TraceRecorder;
diff --git a/Engine/Source/Programs/UnrealInsights/Private/UserInterfaceCommand.cpp b/Engine/Source/Programs/UnrealInsights/Private/UserInterfaceCommand.cpp
--- a/Engine/Source/Programs/UnrealInsights/Private/UserInterfaceCommand.cpp
+++ b/Engine/Source/Programs/UnrealInsights/Private/UserInterfaceCommand.cpp
@@ -44,6 +44,7 @@
 namespace UserInterfaceCommand
 {
 	TSharedRef<FWorkspaceItem> DeveloperTools = FWorkspaceItem::NewGroup(NSLOCTEXT("UnrealInsights", "DeveloperToolsMenu", "Developer Tools"));
+	static void* RecorderEvent = nullptr;
 }
 
 ////////////////////////////////////////////////////////////////////////////////////////////////////
@@ -242,6 +243,15 @@ void FUserInterfaceCommand::InitializeSlateApplication()
 		{
 			const bool bSingleProcess = FParse::Param(FCommandLine::Get(), TEXT("SingleProcess"));
 			TraceInsightsModule.CreateSessionBrowser(bAllowDebugTools, bSingleProcess);
+
+#if PLATFORM_WINDOWS
+			// Create a named event that other processes can use detect a running
+			// recorder and connect to it automatically
+			if (UserInterfaceCommand::RecorderEvent == nullptr)
+			{
+				UserInterfaceCommand::RecorderEvent = ::CreateEvent(nullptr, true, false, TEXT("Local\\UnrealInsightsRecorder"));
+			}
+#endif // PLATFORM_WINDOWS
 		}
 
 		delete[] TraceFile;
@@ -257,6 +267,14 @@ void FUserInterfaceCommand::ShutdownSlateApplication()
 
 	// Shut down application.
 	FSlateApplication::Shutdown();
+
+#if PLATFORM_WINDOWS
+	if (UserInterfaceCommand::RecorderEvent != nullptr)
+	{
+		::CloseHandle(UserInterfaceCommand::RecorderEvent);
+		UserInterfaceCommand::RecorderEvent = nullptr;
+	}
+#endif // PLATFORM_WINDOWS
 }
 
 ////////////////////////////////////////////////////////////////////////////////////////////////////
