[Insights] Fixed slate app initialization for Unreal Insights.

diff --git a/Engine/Source/Programs/UnrealInsights/Private/UserInterfaceCommand.cpp b/Engine/Source/Programs/UnrealInsights/Private/UserInterfaceCommand.cpp
--- a/Engine/Source/Programs/UnrealInsights/Private/UserInterfaceCommand.cpp
+++ b/Engine/Source/Programs/UnrealInsights/Private/UserInterfaceCommand.cpp
@@ -123,15 +123,13 @@ void FUserInterfaceCommand::Run()
 		}
 	}
 
-
-	InitializeSlateApplication();
 	//FCoreStyle::ResetToDefault();
 
 	// Crank up a normal Slate application using the platform's standalone renderer.
 	FSlateApplication::InitializeAsStandaloneApplication(GetStandardStandaloneRenderer());
 
 	// Load required modules.
-
+	FModuleManager::Get().LoadModuleChecked("EditorStyle");
 	FModuleManager::Get().LoadModuleChecked("TraceInsights");
 
 	// Load plug-ins.
@@ -140,8 +138,12 @@ void FUserInterfaceCommand::Run()
 	IPluginManager::Get().LoadModulesForEnabledPlugins(ELoadingPhase::Default);
 
 	// Load optional modules.
-	FModuleManager::Get().LoadModule("SettingsEditor");
+	if (FModuleManager::Get().ModuleExists(TEXT("SettingsEditor")))
+	{
+		FModuleManager::Get().LoadModule("SettingsEditor");
+	}
 
+	InitializeSlateApplication();
 
 	// Initialize source code access.
 	// Load the source code access module.
