Fixed Unreal Insights asserting on start up.

#rnx

diff --git a/Engine/Source/Programs/UnrealInsights/Private/UnrealInsightsMain.cpp b/Engine/Source/Programs/UnrealInsights/Private/UnrealInsightsMain.cpp
--- a/Engine/Source/Programs/UnrealInsights/Private/UnrealInsightsMain.cpp
+++ b/Engine/Source/Programs/UnrealInsights/Private/UnrealInsightsMain.cpp
@@ -2,6 +2,7 @@
 
 #include "UnrealInsightsMain.h"
 
+#include "CoreGlobals.h"
 #include "RequiredProgramMainCPPInclude.h"
 #include "UserInterfaceCommand.h"
 
@@ -14,6 +15,8 @@ IMPLEMENT_APPLICATION(UnrealInsights, "UnrealInsights");
  */
 int32 UnrealInsightsMain(const TCHAR* CommandLine)
 {
+	FTaskTagScope Scope(ETaskTag::EGameThread);
+
 	// Override the stack size for the thread pool.
 	FQueuedThreadPool::OverrideStackSize = 256 * 1024;
 
