[Insights] Fix menu check boxes by incrementing GFrameCounter in the Insights main loop #rb none

diff --git a/Engine/Source/Programs/UnrealInsights/Private/UserInterfaceCommand.cpp b/Engine/Source/Programs/UnrealInsights/Private/UserInterfaceCommand.cpp
--- a/Engine/Source/Programs/UnrealInsights/Private/UserInterfaceCommand.cpp
+++ b/Engine/Source/Programs/UnrealInsights/Private/UserInterfaceCommand.cpp
@@ -5,6 +5,7 @@
 #include "Async/TaskGraphInterfaces.h"
 //#include "Brushes/SlateImageBrush.h"
 #include "Containers/Ticker.h"
+#include "CoreGlobals.h"
 #include "EditorStyleSet.h"
 #include "Framework/Application/SlateApplication.h"
 #include "Framework/Docking/LayoutService.h"
@@ -199,6 +200,8 @@ void FUserInterfaceCommand::Run()
 
 		FCoreDelegates::OnEndFrame.Broadcast();
 		GLog->FlushThreadedLogs(); //im: ???
+
+		GFrameCounter++;
 	}
 
 	//im: ??? FCoreDelegates::OnExit.Broadcast();
