[Insights] Fixed UE4 references.

diff --git a/Engine/Source/Programs/UnrealInsights/Private/Mac/MacUnrealInsightsMain.cpp b/Engine/Source/Programs/UnrealInsights/Private/Mac/MacUnrealInsightsMain.cpp
--- a/Engine/Source/Programs/UnrealInsights/Private/Mac/MacUnrealInsightsMain.cpp
+++ b/Engine/Source/Programs/UnrealInsights/Private/Mac/MacUnrealInsightsMain.cpp
@@ -10,14 +10,14 @@
 static FString GSavedCommandLine;
 
 
-@interface UE4AppDelegate : NSObject<NSApplicationDelegate, NSFileManagerDelegate>
+@interface UE5AppDelegate : NSObject<NSApplicationDelegate, NSFileManagerDelegate>
 {
 }
 
 @end
 
 
-@implementation UE4AppDelegate
+@implementation UE5AppDelegate
 
 //handler for the quit apple event used by the Dock menu
 - (void)handleQuitEvent:(NSAppleEventDescriptor*)Event withReplyEvent:(NSAppleEventDescriptor*)ReplyEvent
@@ -112,7 +112,7 @@ int main(int argc, char *argv[])
 
 	SCOPED_AUTORELEASE_POOL;
 	[NSApplication sharedApplication];
-	[NSApp setDelegate:[UE4AppDelegate new]];
+	[NSApp setDelegate:[UE5AppDelegate new]];
 	[NSApp run];
 	return 0;
 }
