Throttle Insights rendering when application is idle to reduce GPU usage especially when multiple instance are running
  - From 8% GPU to 0.2% per background instance on my 3070

diff --git a/Engine/Source/Programs/UnrealInsights/Private/UserInterfaceCommand.cpp b/Engine/Source/Programs/UnrealInsights/Private/UserInterfaceCommand.cpp
--- a/Engine/Source/Programs/UnrealInsights/Private/UserInterfaceCommand.cpp
+++ b/Engine/Source/Programs/UnrealInsights/Private/UserInterfaceCommand.cpp
@@ -39,12 +39,38 @@
 ////////////////////////////////////////////////////////////////////////////////////////////////////
 
 #define IDEAL_FRAMERATE 60
+#define BACKGROUND_FRAMERATE 4
+#define IDLE_INPUT_SECONDS 5.0f
 
 ////////////////////////////////////////////////////////////////////////////////////////////////////
 
 namespace UserInterfaceCommand
 {
 	TSharedRef<FWorkspaceItem> DeveloperTools = FWorkspaceItem::NewGroup(NSLOCTEXT("UnrealInsights", "DeveloperToolsMenu", "Developer Tools"));
+
+	bool IsApplicationBackground()
+	{
+		return !FPlatformApplicationMisc::IsThisApplicationForeground() && (FPlatformTime::Seconds() - FSlateApplication::Get().GetLastUserInteractionTime()) > IDLE_INPUT_SECONDS;
+	}
+
+	void AdaptiveSleep(float Seconds)
+	{
+		const double IdealFrameTime = 1.0 / IDEAL_FRAMERATE;
+		if (Seconds > IdealFrameTime)
+		{
+			// While in background, pump message at ideal frame time and get out of background as soon as input is received
+			const double WakeupTime = FPlatformTime::Seconds() + Seconds;
+			while (IsApplicationBackground() && FPlatformTime::Seconds() < WakeupTime)
+			{
+				FSlateApplication::Get().PumpMessages();
+				FPlatformProcess::Sleep((float)FMath::Clamp(WakeupTime - FPlatformTime::Seconds(), 0.0, IdealFrameTime));
+			}
+		}
+		else
+		{
+			FPlatformProcess::Sleep(Seconds);
+		}
+	}
 }
 
 ////////////////////////////////////////////////////////////////////////////////////////////////////
@@ -177,6 +203,7 @@ void FUserInterfaceCommand::Run()
 	double DeltaTime = 0.0;
 	double LastTime = FPlatformTime::Seconds();
 	const float IdealFrameTime = 1.0f / IDEAL_FRAMERATE;
+	const float BackgroundFrameTime = 1.0f / BACKGROUND_FRAMERATE;
 
 	while (!IsEngineExitRequested())
 	{
@@ -190,7 +217,8 @@ void FUserInterfaceCommand::Run()
 		FTicker::GetCoreTicker().Tick(DeltaTime);
 
 		// Throttle frame rate.
-		FPlatformProcess::Sleep(FMath::Max<float>(0.0f, IdealFrameTime - (FPlatformTime::Seconds() - LastTime)));
+		const float FrameTime = UserInterfaceCommand::IsApplicationBackground() ? BackgroundFrameTime : IdealFrameTime;
+		UserInterfaceCommand::AdaptiveSleep(FMath::Max<float>(0.0f, FrameTime - (FPlatformTime::Seconds() - LastTime)));
 
 		double CurrentTime = FPlatformTime::Seconds();
 		DeltaTime =  CurrentTime - LastTime;
