[Insights] Fixed the Unreal Insights crash caused by the stats thread and of the task graph being shutdown twice.

diff --git a/Engine/Source/Programs/UnrealInsights/Private/UnrealInsightsMain.cpp b/Engine/Source/Programs/UnrealInsights/Private/UnrealInsightsMain.cpp
--- a/Engine/Source/Programs/UnrealInsights/Private/UnrealInsightsMain.cpp
+++ b/Engine/Source/Programs/UnrealInsights/Private/UnrealInsightsMain.cpp
@@ -6,10 +6,8 @@
 #include "RequiredProgramMainCPPInclude.h"
 #include "UserInterfaceCommand.h"
 
-
 IMPLEMENT_APPLICATION(UnrealInsights, "UnrealInsights");
 
-
 /**
  * Platform agnostic implementation of the main entry point.
  */
@@ -20,14 +18,11 @@ int32 UnrealInsightsMain(const TCHAR* CommandLine)
 	// Override the stack size for the thread pool.
 	FQueuedThreadPool::OverrideStackSize = 256 * 1024;
 
-	//im: ???
-	FCommandLine::Set(CommandLine);
-
 	// Initialize core.
 	GEngineLoop.PreInit(CommandLine);
 
 	// Make sure all UObject classes are registered and default properties have been initialized.
-	//im: ??? ProcessNewlyLoadedUObjects();
+	//ProcessNewlyLoadedUObjects();
 
 	// Tell the module manager it may now process newly-loaded UObjects when new C++ modules are loaded.
 	FModuleManager::Get().StartProcessingNewlyLoadedObjects();
@@ -35,15 +30,9 @@ int32 UnrealInsightsMain(const TCHAR* CommandLine)
 	FUserInterfaceCommand::Run();
 
 	// Shut down.
-	//im: ??? FCoreDelegates::OnExit.Broadcast();
-	FEngineLoop::AppPreExit(); //im: ???
+	//FCoreDelegates::OnExit.Broadcast();
+	FEngineLoop::AppPreExit();
 	FModuleManager::Get().UnloadModulesAtShutdown();
 
-#if STATS
-	FThreadStats::StopThread();
-#endif
-
-	FTaskGraphInterface::Shutdown(); //im: ???
-
 	return 0;
 }
