Fix task tag scoping for UnrealInsights on Mac

#rnx

diff --git a/Engine/Source/Programs/UnrealInsights/Private/Mac/MacUnrealInsightsMain.cpp b/Engine/Source/Programs/UnrealInsights/Private/Mac/MacUnrealInsightsMain.cpp
--- a/Engine/Source/Programs/UnrealInsights/Private/Mac/MacUnrealInsightsMain.cpp
+++ b/Engine/Source/Programs/UnrealInsights/Private/Mac/MacUnrealInsightsMain.cpp
@@ -75,6 +75,11 @@ static FString GSavedCommandLine;
 	}
 }
 
+- (void) applicationWillTerminate:(NSNotification*)notification
+{
+	FTaskTagScope::SetTagStaticInit();
+}
+
 - (void)applicationDidFinishLaunching:(NSNotification *)Notification
 {
 	//install the custom quit event handler
@@ -89,6 +94,8 @@ static FString GSavedCommandLine;
 
 int main(int argc, char *argv[])
 {
+	FTaskTagScope::SetTagNone();
+
 	for (int32 Option = 1; Option < argc; Option++)
 	{
 		GSavedCommandLine += TEXT(" ");
