[Insights] Add dependency on TargetPlatform module because it is needed for the automation module to work.

diff --git a/Engine/Source/Programs/UnrealInsights/UnrealInsights.Build.cs b/Engine/Source/Programs/UnrealInsights/UnrealInsights.Build.cs
--- a/Engine/Source/Programs/UnrealInsights/UnrealInsights.Build.cs
+++ b/Engine/Source/Programs/UnrealInsights/UnrealInsights.Build.cs
@@ -21,6 +21,7 @@ public class UnrealInsights : ModuleRules
 				"SlateCore",
 				"SourceCodeAccess",
 				"StandaloneRenderer",
+				"TargetPlatform",
 				"TraceInsights",
 			}
 		);
