[Insights]
  Updated the name of the Windows event used to force single instance of UnrealInsights.exe (Browser Mode) and to control when to try to auto-connect. This change allows UnrealInsights.exe from UE5 to coexists with a running instance from UE4.
  Also updated the comments to better explain how the Windows event is used.
  Trace auto-connect behavior remains the same: a game runtime (editor excluded) will try to auto-connect with a local trace store (now hosted by the UnrealTraceServer.exe) if UnrealInsights (Browser Mode) is running.

diff --git a/Engine/Source/Programs/UnrealInsights/Private/UserInterfaceCommand.cpp b/Engine/Source/Programs/UnrealInsights/Private/UserInterfaceCommand.cpp
--- a/Engine/Source/Programs/UnrealInsights/Private/UserInterfaceCommand.cpp
+++ b/Engine/Source/Programs/UnrealInsights/Private/UserInterfaceCommand.cpp
@@ -78,9 +78,11 @@ namespace UserInterfaceCommand
 bool CheckSessionBrowserSingleInstance()
 {
 #if PLATFORM_WINDOWS
-	// Create a named event that other processes can use to detect a running recorder and connect to it automatically.
-	// See usage in \Engine\Source\Runtime\Launch\Private\LaunchEngineLoop.cpp
-	HANDLE SessionBrowserEvent = CreateEvent(NULL, true, false, TEXT("Local\\UnrealInsightsRecorder"));
+	// Create a named event that other processes can detect.
+	// It allows only a single instance of Unreal Insights (Browser Mode).
+	// The event is also used by runtime to choose when to try to auto-connect.
+	// See FTraceAuxiliary::TryAutoConnect() in \Runtime\Core\Private\ProfilingDebugging\TraceAuxiliary.cpp
+	HANDLE SessionBrowserEvent = CreateEvent(NULL, true, false, TEXT("Local\\UnrealInsightsBrowser"));
 	if (SessionBrowserEvent == NULL || GetLastError() == ERROR_ALREADY_EXISTS)
 	{
 		// Another Session Browser process is already running.
@@ -111,7 +113,7 @@ bool CheckSessionBrowserSingleInstance()
 #endif // PLATFORM_WINDOWS
 
 #if PLATFORM_UNIX
-	int FileHandle = open("/var/run/UnrealInsights.pid", O_CREAT | O_RDWR, 0666);
+	int FileHandle = open("/var/run/UnrealInsightsBrowser.pid", O_CREAT | O_RDWR, 0666);
 	int Ret = flock(FileHandle, LOCK_EX | LOCK_NB);
 	if (Ret && EWOULDBLOCK == errno)
 	{
diff --git a/Engine/Source/Runtime/Core/Private/ProfilingDebugging/TraceAuxiliary.cpp b/Engine/Source/Runtime/Core/Private/ProfilingDebugging/TraceAuxiliary.cpp
--- a/Engine/Source/Runtime/Core/Private/ProfilingDebugging/TraceAuxiliary.cpp
+++ b/Engine/Source/Runtime/Core/Private/ProfilingDebugging/TraceAuxiliary.cpp
@@ -691,8 +691,9 @@ void FTraceAuxiliary::TryAutoConnect()
 {
 #if UE_TRACE_ENABLED
 	#if PLATFORM_WINDOWS
-	// If we can detect a named event then we can try and auto-connect to UnrealInsights.
-	HANDLE KnownEvent = ::OpenEvent(EVENT_ALL_ACCESS, false, TEXT("Local\\UnrealInsightsRecorder"));
+	// If we can detect a named event it means UnrealInsights (Browser Mode) is running.
+	// In this case, we try to auto-connect with the Trace Server.
+	HANDLE KnownEvent = ::OpenEvent(EVENT_ALL_ACCESS, false, TEXT("Local\\UnrealInsightsBrowser"));
 	if (KnownEvent != nullptr)
 	{
 		GTraceAuxiliary.Connect(ETraceConnectType::Network, TEXT("127.0.0.1"));
