[Insights] Added source files and script for app icons (Mac).

diff --git a/Engine/Source/Programs/UnrealInsights/Resources/Mac/make_icns.sh b/Engine/Source/Programs/UnrealInsights/Resources/Mac/make_icns.sh
new file mode 100644
--- /dev/null
+++ b/Engine/Source/Programs/UnrealInsights/Resources/Mac/make_icns.sh
@@ -0,0 +1,21 @@
+#!/bin/sh
+# Copyright Epic Games, Inc. All Rights Reserved.
+
+# Content of each /*.iconset/ folder:
+# File                 Size       DPI
+# -----------------------------------
+# icon_16x16.png       16x16       72
+# [REDACTED_EMAIL]    32x32      144
+# icon_32x32.png       32x32       72
+# [REDACTED_EMAIL]    64x64      144
+# icon_128x128.png     128x128     72
+# [REDACTED_EMAIL]  256x256    144
+# icon_256x256.png     256x256     72
+# [REDACTED_EMAIL]  512x512    144
+# icon_512x512.png     512x512     72
+# [REDACTED_EMAIL]  1024x1024  144
+# -----------------------------------
+# https://developer.apple.com/design/human-interface-guidelines/macos/icons-and-images/app-icon/
+
+iconutil --convert icns UnrealInsights.iconset
+iconutil --convert icns UTrace.iconset
