[Insights] Added source files and script for app icons (Windows).

diff --git a/Engine/Source/Programs/UnrealInsights/Resources/Windows/make_ico.py b/Engine/Source/Programs/UnrealInsights/Resources/Windows/make_ico.py
new file mode 100644
--- /dev/null
+++ b/Engine/Source/Programs/UnrealInsights/Resources/Windows/make_ico.py
@@ -0,0 +1,16 @@
+# Copyright Epic Games, Inc. All Rights Reserved.
+
+from PIL import Image
+
+def MakeIcon(folder, name, output):
+    sizes = [256, 48, 32, 24, 16]
+    images = []
+    for size in sizes:
+        print(f"Loading {folder}/icon_{size}x{size}.png")
+        image = Image.open(f"{folder}/icon_{size}x{size}.png")
+        images.append(image)
+    print(f"Creating {output}")
+    images[0].save(f"{output}", append_images=images)
+
+MakeIcon("UnrealInsights", "UnrealInsights", "UnrealInsights.ico")
+MakeIcon("UTrace", "UTrace", "UTrace.ico")
