[Insights] Fixed Unreal Insights (Session Browser) app to re-focus if it was already running.

diff --git a/Engine/Source/Programs/UnrealInsights/Private/UserInterfaceCommand.cpp b/Engine/Source/Programs/UnrealInsights/Private/UserInterfaceCommand.cpp
--- a/Engine/Source/Programs/UnrealInsights/Private/UserInterfaceCommand.cpp
+++ b/Engine/Source/Programs/UnrealInsights/Private/UserInterfaceCommand.cpp
@@ -93,7 +93,7 @@ bool CheckSessionBrowserSingleInstance()
 		}
 
 		// Activate the respective window.
-		HWND Window = FindWindowW(0, L"Unreal Insights");
+		HWND Window = FindWindowW(0, L"Unreal Insights Session Browser");
 		if (Window)
 		{
 			ShowWindow(Window, SW_SHOW);
