#UnrealInsights Add compile option to build UnrealInsights monolithic. Note that there is already a generic -Monolithic commandline parameter to TargetRules, but when explicitly setting the LinkType it is ignored, so we are also explicitly supporting -Monolithic.

diff --git a/Engine/Source/Programs/UnrealInsights/UnrealInsights.Target.cs b/Engine/Source/Programs/UnrealInsights/UnrealInsights.Target.cs
--- a/Engine/Source/Programs/UnrealInsights/UnrealInsights.Target.cs
+++ b/Engine/Source/Programs/UnrealInsights/UnrealInsights.Target.cs
@@ -2,14 +2,18 @@
 
 using UnrealBuildTool;
 using System.Collections.Generic;
+using EpicGames.Core;
 
 [SupportedPlatforms("Win64", "Linux", "Mac")]
 public class UnrealInsightsTarget : TargetRules
 {
+	[CommandLine("-Monolithic")]
+	public bool bMonolithic = false;
+
 	public UnrealInsightsTarget(TargetInfo Target) : base(Target)
 	{
 		Type = TargetType.Program;
-		LinkType = TargetLinkType.Modular; // TargetLinkType.Monolithic;
+		LinkType = bMonolithic ? TargetLinkType.Monolithic : TargetLinkType.Modular;
 
 		LaunchModuleName = "UnrealInsights";
 		ExtraModuleNames.Add("EditorStyle");
