Enable utrace file association with Insights on MAC.

diff --git a/Engine/Source/Programs/UnrealInsights/Private/Mac/MacUnrealInsightsMain.cpp b/Engine/Source/Programs/UnrealInsights/Private/Mac/MacUnrealInsightsMain.cpp
--- a/Engine/Source/Programs/UnrealInsights/Private/Mac/MacUnrealInsightsMain.cpp
+++ b/Engine/Source/Programs/UnrealInsights/Private/Mac/MacUnrealInsightsMain.cpp
@@ -12,6 +12,7 @@ static FString GSavedCommandLine;
 
 @interface UE5AppDelegate : NSObject<NSApplicationDelegate, NSFileManagerDelegate>
 {
+	bool bHasFinishedLaunching;
 }
 
 @end
@@ -30,6 +31,45 @@ static FString GSavedCommandLine;
 	RequestEngineExit(TEXT("requestQuit"));
 }
 
+- (void)awakeFromNib
+{
+	bHasFinishedLaunching = false;
+}
+
+- (BOOL)application : (NSApplication*)theApplication openFile : (NSString*)filename
+{
+	FString StringFilename(filename);
+	if(!StringFilename.EndsWith(TEXT(".utrace")))
+	{
+		return NO;
+	}
+	if (!bHasFinishedLaunching)
+	{
+		if ([[NSFileManager defaultManager]fileExistsAtPath:filename] )
+		{
+			GSavedCommandLine = GSavedCommandLine + TEXT(" ") + StringFilename;
+			return YES;
+		}
+		return NO;
+	}
+	else if ([[NSFileManager defaultManager]fileExistsAtPath:filename] )
+	{
+		NSURL* BundleURL = [[NSRunningApplication currentApplication]bundleURL];
+
+		NSDictionary* Configuration = [NSDictionary dictionaryWithObject : [NSArray arrayWithObject : filename] forKey : NSWorkspaceLaunchConfigurationArguments];
+
+		NSError* Error = nil;
+
+		NSRunningApplication* NewInstance = [[NSWorkspace sharedWorkspace]launchApplicationAtURL:BundleURL options : (NSWorkspaceLaunchOptions)(NSWorkspaceLaunchAsync | NSWorkspaceLaunchNewInstance) configuration : Configuration error : &Error];
+
+		return (NewInstance != nil);
+	}
+	else
+	{
+		return NO;
+	}
+}
+
 - (void) runGameThread:(id)Arg
 {
 	FPlatformMisc::SetGracefulTerminationHandler();
@@ -86,6 +126,8 @@ static FString GSavedCommandLine;
 	NSAppleEventManager* appleEventManager = [NSAppleEventManager sharedAppleEventManager];
 	[appleEventManager setEventHandler:self andSelector:@selector(handleQuitEvent:withReplyEvent:) forEventClass:kCoreEventClass andEventID:kAEQuitApplication];
 
+	bHasFinishedLaunching = true;
+
 	RunGameThread(self, @selector(runGameThread:));
 }
 
