- Fixed up various threading and menu initialization issues with Programs on Mac
- Moved boilerplate program startup code into a shared Obj-C class, MacProgramDelegate

diff --git a/Engine/Source/Programs/SlateUGS/Private/Platforms/Mac/SlateUGSMainMac.cpp b/Engine/Source/Programs/SlateUGS/Private/Platforms/Mac/SlateUGSMainMac.cpp
--- a/Engine/Source/Programs/SlateUGS/Private/Platforms/Mac/SlateUGSMainMac.cpp
+++ b/Engine/Source/Programs/SlateUGS/Private/Platforms/Mac/SlateUGSMainMac.cpp
@@ -1,85 +1,12 @@
 // Copyright Epic Games, Inc. All Rights Reserved.
 
 #include "SlateUGSApp.h"
-#include "HAL/ExceptionHandling.h"
-#include "Mac/CocoaThread.h"
+#include "Mac/MacProgramDelegate.h"
+#include "LaunchEngineLoop.h"
 
-static FString GSavedCommandLine;
-
-@interface UEAppDelegate : NSObject <NSApplicationDelegate, NSFileManagerDelegate>
-{
-}
-
-@end
-
-@implementation UEAppDelegate
-
-//handler for the quit apple event used by the Dock menu
-- (void)handleQuitEvent:(NSAppleEventDescriptor*)Event withReplyEvent:(NSAppleEventDescriptor*)ReplyEvent
-{
-    [NSApp terminate:self];
-}
-
-- (void) runGameThread:(id)Arg
-{
-	FPlatformMisc::SetGracefulTerminationHandler();
-	FPlatformMisc::SetCrashHandler(nullptr);
-	
-	RunSlateUGS(*GSavedCommandLine);
-	
-	[NSApp terminate: self];
-}
-
-- (NSApplicationTerminateReply)applicationShouldTerminate:(NSApplication *)Sender;
-{
-	if(!IsEngineExitRequested() || ([NSThread gameThread] && [NSThread gameThread] != [NSThread mainThread]))
-	{
-		RequestEngineExit(TEXT("applicationShouldTerminate"));
-		return NSTerminateLater;
-	}
-	else
-	{
-		return NSTerminateNow;
-	}
-}
-
-- (void)applicationDidFinishLaunching:(NSNotification *)Notification
+int main(int argc, char *argv[])
 {
-	//install the custom quit event handler
-    NSAppleEventManager* appleEventManager = [NSAppleEventManager sharedAppleEventManager];
-    [appleEventManager setEventHandler:self andSelector:@selector(handleQuitEvent:withReplyEvent:) forEventClass:kCoreEventClass andEventID:kAEQuitApplication];
-	
-	RunGameThread(self, @selector(runGameThread:));
+	[MacProgramDelegate mainWithArgc:argc argv:argv programMain:RunSlateUGS programExit:FEngineLoop::AppExit];
 }
 
-@end
 
-int main(int argc, char *argv[])
-{
-	for (int32 Option = 1; Option < argc; Option++)
-	{
-		GSavedCommandLine += TEXT(" ");
-		FString Argument(ANSI_TO_TCHAR(argv[Option]));
-		if (Argument.Contains(TEXT(" ")))
-		{
-			if (Argument.Contains(TEXT("=")))
-			{
-				FString ArgName;
-				FString ArgValue;
-				Argument.Split( TEXT("="), &ArgName, &ArgValue );
-				Argument = FString::Printf( TEXT("%s=\"%s\""), *ArgName, *ArgValue );
-			}
-			else
-			{
-				Argument = FString::Printf(TEXT("\"%s\""), *Argument);
-			}
-		}
-		GSavedCommandLine += Argument;
-	}
-
-	SCOPED_AUTORELEASE_POOL;
-	[NSApplication sharedApplication];
-	[NSApp setDelegate:[UEAppDelegate new]];
-	[NSApp run];
-	return 0;
-}
diff --git a/Engine/Source/Programs/SlateViewer/Private/Mac/SlateViewerMainMac.cpp b/Engine/Source/Programs/SlateViewer/Private/Mac/SlateViewerMainMac.cpp
--- a/Engine/Source/Programs/SlateViewer/Private/Mac/SlateViewerMainMac.cpp
+++ b/Engine/Source/Programs/SlateViewer/Private/Mac/SlateViewerMainMac.cpp
@@ -1,85 +1,10 @@
 // Copyright Epic Games, Inc. All Rights Reserved.
 
 #include "SlateViewerApp.h"
-#include "HAL/ExceptionHandling.h"
-#include "Mac/CocoaThread.h"
-
-static FString GSavedCommandLine;
-
-@interface UEAppDelegate : NSObject <NSApplicationDelegate, NSFileManagerDelegate>
-{
-}
-
-@end
-
-@implementation UEAppDelegate
-
-//handler for the quit apple event used by the Dock menu
-- (void)handleQuitEvent:(NSAppleEventDescriptor*)Event withReplyEvent:(NSAppleEventDescriptor*)ReplyEvent
-{
-    [NSApp terminate:self];
-}
-
-- (void) runGameThread:(id)Arg
-{
-	FPlatformMisc::SetGracefulTerminationHandler();
-	FPlatformMisc::SetCrashHandler(nullptr);
-	
-	RunSlateViewer(*GSavedCommandLine);
-	
-	[NSApp terminate: self];
-}
-
-- (NSApplicationTerminateReply)applicationShouldTerminate:(NSApplication *)Sender;
-{
-	if(!IsEngineExitRequested() || ([NSThread gameThread] && [NSThread gameThread] != [NSThread mainThread]))
-	{
-		RequestEngineExit(TEXT("applicationShouldTerminate"));
-		return NSTerminateLater;
-	}
-	else
-	{
-		return NSTerminateNow;
-	}
-}
-
-- (void)applicationDidFinishLaunching:(NSNotification *)Notification
-{
-	//install the custom quit event handler
-    NSAppleEventManager* appleEventManager = [NSAppleEventManager sharedAppleEventManager];
-    [appleEventManager setEventHandler:self andSelector:@selector(handleQuitEvent:withReplyEvent:) forEventClass:kCoreEventClass andEventID:kAEQuitApplication];
-	
-	RunGameThread(self, @selector(runGameThread:));
-}
-
-@end
+#include "Mac/MacProgramDelegate.h"
+#include "LaunchEngineLoop.h"
 
 int main(int argc, char *argv[])
 {
-	for (int32 Option = 1; Option < argc; Option++)
-	{
-		GSavedCommandLine += TEXT(" ");
-		FString Argument(ANSI_TO_TCHAR(argv[Option]));
-		if (Argument.Contains(TEXT(" ")))
-		{
-			if (Argument.Contains(TEXT("=")))
-			{
-				FString ArgName;
-				FString ArgValue;
-				Argument.Split( TEXT("="), &ArgName, &ArgValue );
-				Argument = FString::Printf( TEXT("%s=\"%s\""), *ArgName, *ArgValue );
-			}
-			else
-			{
-				Argument = FString::Printf(TEXT("\"%s\""), *Argument);
-			}
-		}
-		GSavedCommandLine += Argument;
-	}
-
-	SCOPED_AUTORELEASE_POOL;
-	[NSApplication sharedApplication];
-	[NSApp setDelegate:[UEAppDelegate new]];
-	[NSApp run];
-	return 0;
+	return [MacProgramDelegate mainWithArgc:argc argv:argv programMain:RunSlateViewer programExit:FEngineLoop::AppExit];
 }
diff --git a/Engine/Source/Programs/UnrealFrontend/Private/Mac/MacUnrealFrontendMain.cpp b/Engine/Source/Programs/UnrealFrontend/Private/Mac/MacUnrealFrontendMain.cpp
--- a/Engine/Source/Programs/UnrealFrontend/Private/Mac/MacUnrealFrontendMain.cpp
+++ b/Engine/Source/Programs/UnrealFrontend/Private/Mac/MacUnrealFrontendMain.cpp
@@ -1,119 +1,10 @@
 // Copyright Epic Games, Inc. All Rights Reserved.
 
-#include "CoreMinimal.h"
 #include "UnrealFrontendMain.h"
-#include "HAL/ExceptionHandling.h"
+#include "Mac/MacProgramDelegate.h"
 #include "LaunchEngineLoop.h"
-#include "Misc/CommandLine.h"
-#include "Mac/CocoaThread.h"
-
-
-static FString GSavedCommandLine;
-
-
-@interface UnrealAppDelegate : NSObject<NSApplicationDelegate, NSFileManagerDelegate>
-{
-}
-
-@end
-
-
-@implementation UnrealAppDelegate
-
-//handler for the quit apple event used by the Dock menu
-- (void)handleQuitEvent:(NSAppleEventDescriptor*)Event withReplyEvent:(NSAppleEventDescriptor*)ReplyEvent
-{
-    [self requestQuit:self];
-}
-
-- (IBAction)requestQuit:(id)Sender
-{
-	RequestEngineExit(TEXT("requestQuit"));
-}
-
-- (void)runGameThread:(id)Arg
-{
-	FPlatformMisc::SetGracefulTerminationHandler();
-	FPlatformMisc::SetCrashHandler(nullptr);
-	
-#if !UE_BUILD_SHIPPING
-	if (FParse::Param(*GSavedCommandLine,TEXT("crashreports")))
-	{
-		GAlwaysReportCrash = true;
-	}
-#endif
-	
-#if UE_BUILD_DEBUG
-	if (!GAlwaysReportCrash)
-#else
-		if (FPlatformMisc::IsDebuggerPresent() && !GAlwaysReportCrash)
-#endif
-		{
-			UnrealFrontendMain(*GSavedCommandLine);
-		}
-		else
-		{
-			GIsGuarded = 1;
-			UnrealFrontendMain(*GSavedCommandLine);
-			GIsGuarded = 0;
-		}
-	
-	FEngineLoop::AppExit();
-	
-	[NSApp terminate: self];
-}
-
-- (NSApplicationTerminateReply)applicationShouldTerminate:(NSApplication *)Sender;
-{
-	if(!IsEngineExitRequested() || ([NSThread gameThread] && [NSThread gameThread] != [NSThread mainThread]))
-	{
-		[self requestQuit:self];
-		return NSTerminateLater;
-	}
-	else
-	{
-		return NSTerminateNow;
-	}
-}
-
-- (void)applicationDidFinishLaunching:(NSNotification *)Notification
-{
-	//install the custom quit event handler
-    NSAppleEventManager* appleEventManager = [NSAppleEventManager sharedAppleEventManager];
-    [appleEventManager setEventHandler:self andSelector:@selector(handleQuitEvent:withReplyEvent:) forEventClass:kCoreEventClass andEventID:kAEQuitApplication];
-	
-	RunGameThread(self, @selector(runGameThread:));
-}
-
-@end
-
 
 int main( int argc, char *argv[] )
 {
-	for (int32 Option = 1; Option < argc; Option++)
-	{
-		GSavedCommandLine += TEXT(" ");
-		FString Argument(ANSI_TO_TCHAR(argv[Option]));
-		if (Argument.Contains(TEXT(" ")))
-		{
-			if (Argument.Contains(TEXT("=")))
-			{
-				FString ArgName;
-				FString ArgValue;
-				Argument.Split( TEXT("="), &ArgName, &ArgValue );
-				Argument = FString::Printf( TEXT("%s=\"%s\""), *ArgName, *ArgValue );
-			}
-			else
-			{
-				Argument = FString::Printf(TEXT("\"%s\""), *Argument);
-			}
-		}
-		GSavedCommandLine += Argument;
-	}
-
-	SCOPED_AUTORELEASE_POOL;
-	[NSApplication sharedApplication];
-	[NSApp setDelegate:[UnrealAppDelegate new]];
-	[NSApp run];
-	return 0;
+	return [MacProgramDelegate mainWithArgc:argc argv:argv programMain:UnrealFrontendMain programExit:FEngineLoop::AppExit];
 }
diff --git a/Engine/Source/Programs/UnrealFrontend/Private/UnrealFrontendMain.cpp b/Engine/Source/Programs/UnrealFrontend/Private/UnrealFrontendMain.cpp
--- a/Engine/Source/Programs/UnrealFrontend/Private/UnrealFrontendMain.cpp
+++ b/Engine/Source/Programs/UnrealFrontend/Private/UnrealFrontendMain.cpp
@@ -17,6 +17,8 @@ IMPLEMENT_APPLICATION(UnrealFrontend, "UnrealFrontend");
  */
 int32 UnrealFrontendMain( const TCHAR* CommandLine )
 {
+	FTaskTagScope TaskTagScope(ETaskTag::EGameThread);
+
 	// Override the stack size for the thread pool.
 	FQueuedThreadPool::OverrideStackSize = 256 * 1024;
 
diff --git a/Engine/Source/Programs/UnrealInsights/Private/Mac/MacUnrealInsightsMain.cpp b/Engine/Source/Programs/UnrealInsights/Private/Mac/MacUnrealInsightsMain.cpp
--- a/Engine/Source/Programs/UnrealInsights/Private/Mac/MacUnrealInsightsMain.cpp
+++ b/Engine/Source/Programs/UnrealInsights/Private/Mac/MacUnrealInsightsMain.cpp
@@ -1,16 +1,10 @@
 // Copyright Epic Games, Inc. All Rights Reserved.
 
 #include "UnrealInsightsMain.h"
-#include "HAL/ExceptionHandling.h"
+#include "Mac/MacProgramDelegate.h"
 #include "LaunchEngineLoop.h"
-#include "Misc/CommandLine.h"
-#include "Mac/CocoaThread.h"
 
-
-static FString GSavedCommandLine;
-
-
-@interface UE5AppDelegate : NSObject<NSApplicationDelegate, NSFileManagerDelegate>
+@interface InsightsAppDelegate : MacProgramDelegate
 {
 	bool bHasFinishedLaunching;
 }
@@ -18,18 +12,7 @@ static FString GSavedCommandLine;
 @end
 
 
-@implementation UE5AppDelegate
-
-//handler for the quit apple event used by the Dock menu
-- (void)handleQuitEvent:(NSAppleEventDescriptor*)Event withReplyEvent:(NSAppleEventDescriptor*)ReplyEvent
-{
-	[self requestQuit:self];
-}
-
-- (IBAction)requestQuit:(id)Sender
-{
-	RequestEngineExit(TEXT("requestQuit"));
-}
+@implementation InsightsAppDelegate
 
 - (void)awakeFromNib
 {
@@ -47,7 +30,7 @@ static FString GSavedCommandLine;
 	{
 		if ([[NSFileManager defaultManager]fileExistsAtPath:filename] )
 		{
-			GSavedCommandLine = GSavedCommandLine + TEXT(" ") + StringFilename;
+			SavedCommandLine = SavedCommandLine + TEXT(" ") + StringFilename;
 			return YES;
 		}
 		return NO;
@@ -70,65 +53,11 @@ static FString GSavedCommandLine;
 	}
 }
 
-- (void) runGameThread:(id)Arg
-{
-	FPlatformMisc::SetGracefulTerminationHandler();
-	FPlatformMisc::SetCrashHandler(nullptr);
-	
-#if !UE_BUILD_SHIPPING
-	if (FParse::Param(*GSavedCommandLine,TEXT("crashreports")))
-	{
-		GAlwaysReportCrash = true;
-	}
-#endif
-	
-#if UE_BUILD_DEBUG
-	if (!GAlwaysReportCrash)
-#else
-		if (FPlatformMisc::IsDebuggerPresent() && !GAlwaysReportCrash)
-#endif
-		{
-			UnrealInsightsMain(*GSavedCommandLine);
-		}
-		else
-		{
-			GIsGuarded = 1;
-			UnrealInsightsMain(*GSavedCommandLine);
-			GIsGuarded = 0;
-		}
-	
-	FEngineLoop::AppExit();
-
-	[NSApp terminate: self];
-}
-
-- (NSApplicationTerminateReply)applicationShouldTerminate:(NSApplication *)Sender;
-{
-	if(!IsEngineExitRequested() || ([NSThread gameThread] && [NSThread gameThread] != [NSThread mainThread]))
-	{
-		[self requestQuit:self];
-		return NSTerminateLater;
-	}
-	else
-	{
-		return NSTerminateNow;
-	}
-}
-
-- (void) applicationWillTerminate:(NSNotification*)notification
-{
-	FTaskTagScope::SetTagStaticInit();
-}
-
 - (void)applicationDidFinishLaunching:(NSNotification *)Notification
 {
-	//install the custom quit event handler
-	NSAppleEventManager* appleEventManager = [NSAppleEventManager sharedAppleEventManager];
-	[appleEventManager setEventHandler:self andSelector:@selector(handleQuitEvent:withReplyEvent:) forEventClass:kCoreEventClass andEventID:kAEQuitApplication];
-
 	bHasFinishedLaunching = true;
-
-	RunGameThread(self, @selector(runGameThread:));
+	
+	[super applicationDidFinishLaunching:Notification];
 }
 
 @end
@@ -136,32 +65,8 @@ static FString GSavedCommandLine;
 
 int main(int argc, char *argv[])
 {
-	FTaskTagScope::SetTagNone();
-
-	for (int32 Option = 1; Option < argc; Option++)
-	{
-		GSavedCommandLine += TEXT(" ");
-		FString Argument(ANSI_TO_TCHAR(argv[Option]));
-		if (Argument.Contains(TEXT(" ")))
-		{
-			if (Argument.Contains(TEXT("=")))
-			{
-				FString ArgName;
-				FString ArgValue;
-				Argument.Split( TEXT("="), &ArgName, &ArgValue );
-				Argument = FString::Printf( TEXT("%s=\"%s\""), *ArgName, *ArgValue );
-			}
-			else
-			{
-				Argument = FString::Printf(TEXT("\"%s\""), *Argument);
-			}
-		}
-		GSavedCommandLine += Argument;
-	}
-
-	SCOPED_AUTORELEASE_POOL;
-	[NSApplication sharedApplication];
-	[NSApp setDelegate:[UE5AppDelegate new]];
-	[NSApp run];
-	return 0;
+	// make custom delegate
+	InsightsAppDelegate* Delegate = [[InsightsAppDelegate alloc] initWithProgramMain:UnrealInsightsMain programExit:FEngineLoop::AppExit];
+	// run with it
+	return [MacProgramDelegate mainWithArgc:argc argv:argv existingDelegate:Delegate];
 }
diff --git a/Engine/Source/Programs/UnrealMultiUserSlateServer/Private/Mac/MacUnrealMultiUserSlateServerMain.cpp b/Engine/Source/Programs/UnrealMultiUserSlateServer/Private/Mac/MacUnrealMultiUserSlateServerMain.cpp
--- a/Engine/Source/Programs/UnrealMultiUserSlateServer/Private/Mac/MacUnrealMultiUserSlateServerMain.cpp
+++ b/Engine/Source/Programs/UnrealMultiUserSlateServer/Private/Mac/MacUnrealMultiUserSlateServerMain.cpp
@@ -1,124 +1,10 @@
 // Copyright Epic Games, Inc. All Rights Reserved.
 
-#include "CoreMinimal.h"
 #include "UnrealMultiUserSlateServerRun.h"
-#include "HAL/ExceptionHandling.h"
+#include "Mac/MacProgramDelegate.h"
 #include "LaunchEngineLoop.h"
-#include "Misc/CommandLine.h"
-#include "Mac/CocoaThread.h"
-
-
-static FString GSavedCommandLine;
-
-
-@interface UnrealAppDelegate : NSObject<NSApplicationDelegate, NSFileManagerDelegate>
-{
-}
-
-@end
-
-
-@implementation UnrealAppDelegate
-
-//handler for the quit apple event used by the Dock menu
-- (void)handleQuitEvent:(NSAppleEventDescriptor*)Event withReplyEvent:(NSAppleEventDescriptor*)ReplyEvent
-{
-    [self requestQuit:self];
-}
-
-- (IBAction)requestQuit:(id)Sender
-{
-	RequestEngineExit(TEXT("requestQuit"));
-}
-
-- (void)runGameThread:(id)Arg
-{
-	FPlatformMisc::SetGracefulTerminationHandler();
-	FPlatformMisc::SetCrashHandler(nullptr);
-	
-#if !UE_BUILD_SHIPPING
-	if (FParse::Param(*GSavedCommandLine,TEXT("crashreports")))
-	{
-		GAlwaysReportCrash = true;
-	}
-#endif
-	
-#if UE_BUILD_DEBUG
-	if (!GAlwaysReportCrash)
-#else
-		if (FPlatformMisc::IsDebuggerPresent() && !GAlwaysReportCrash)
-#endif
-		{
-			RunUnrealMultiUserServer(*GSavedCommandLine);
-		}
-		else
-		{
-			GIsGuarded = 1;
-			RunUnrealMultiUserServer(*GSavedCommandLine);
-			GIsGuarded = 0;
-		}
-	
-	[NSApp terminate: self];
-}
-
-- (NSApplicationTerminateReply)applicationShouldTerminate:(NSApplication *)Sender;
-{
-	if(!IsEngineExitRequested() || ([NSThread gameThread] && [NSThread gameThread] != [NSThread mainThread]))
-	{
-		[self requestQuit:self];
-		return NSTerminateLater;
-	}
-	else
-	{
-		return NSTerminateNow;
-	}
-}
-
-- (void) applicationWillTerminate:(NSNotification*)notification
-{
-	FTaskTagScope::SetTagStaticInit();
-}
-
-- (void)applicationDidFinishLaunching:(NSNotification *)Notification
-{
-	//install the custom quit event handler
-    NSAppleEventManager* appleEventManager = [NSAppleEventManager sharedAppleEventManager];
-    [appleEventManager setEventHandler:self andSelector:@selector(handleQuitEvent:withReplyEvent:) forEventClass:kCoreEventClass andEventID:kAEQuitApplication];
-	
-	RunGameThread(self, @selector(runGameThread:));
-}
-
-@end
-
 
 int main( int argc, char *argv[] )
 {
-	FTaskTagScope::SetTagNone();
-
-	for (int32 Option = 1; Option < argc; Option++)
-	{
-		GSavedCommandLine += TEXT(" ");
-		FString Argument(ANSI_TO_TCHAR(argv[Option]));
-		if (Argument.Contains(TEXT(" ")))
-		{
-			if (Argument.Contains(TEXT("=")))
-			{
-				FString ArgName;
-				FString ArgValue;
-				Argument.Split( TEXT("="), &ArgName, &ArgValue );
-				Argument = FString::Printf( TEXT("%s=\"%s\""), *ArgName, *ArgValue );
-			}
-			else
-			{
-				Argument = FString::Printf(TEXT("\"%s\""), *Argument);
-			}
-		}
-		GSavedCommandLine += Argument;
-	}
-
-	SCOPED_AUTORELEASE_POOL;
-	[NSApplication sharedApplication];
-	[NSApp setDelegate:[UnrealAppDelegate new]];
-	[NSApp run];
-	return 0;
+	return [MacProgramDelegate mainWithArgc:argc argv:argv programMain:RunUnrealMultiUserServer programExit:FEngineLoop::AppExit];
 }
diff --git a/Engine/Source/Runtime/ApplicationCore/Private/Mac/Accessibility/CocoaAccessibilityView.cpp b/Engine/Source/Runtime/ApplicationCore/Private/Mac/Accessibility/CocoaAccessibilityView.cpp
--- a/Engine/Source/Runtime/ApplicationCore/Private/Mac/Accessibility/CocoaAccessibilityView.cpp
+++ b/Engine/Source/Runtime/ApplicationCore/Private/Mac/Accessibility/CocoaAccessibilityView.cpp
@@ -25,7 +25,7 @@
 
 - (void)SetAccessibilityWindowAsAccessibilityChild:(AccessibleWidgetId)InWindowId
 {
-	checkf(!IsInGameThread(), TEXT("Updating accessibility view from Game Thread! Accessibility  can only be done on Main Thread!"));
+	checkf([NSThread isMainThread], TEXT("Updating accessibility view in FCocoaWindow from wrong Thread! Accessibility can only be done on Main Thread!"));
 	checkf(InWindowId != IAccessibleWidget::InvalidAccessibleWidgetId, TEXT("Cannot update view with invalid accessible widget Id! use RemoveAccessibilityWindow if trying to remove window."));
 	FMacAccessibilityElement* InAccessibleWindow = [[FMacAccessibilityManager AccessibilityManager]GetAccessibilityElement:InWindowId];
 	check(InAccessibleWindow != Nil);
@@ -37,7 +37,7 @@
 
 - (void)RemoveAccessibilityWindow
 {
-	checkf(!IsInGameThread(), TEXT("Clearing accessibility view from Game Thread! Accessibility  can only be done on Main Thread!"));
+	checkf([NSThread isMainThread], TEXT("Updating accessibility view in FCocoaWindow from wrong Thread! Accessibility can only be done on Main Thread!"));
 	if(self.AccessibleWindow != Nil)
 	{
 		self.AccessibleWindow.accessibilityParent = Nil;
diff --git a/Engine/Source/Runtime/ApplicationCore/Private/Mac/CocoaWindow.cpp b/Engine/Source/Runtime/ApplicationCore/Private/Mac/CocoaWindow.cpp
--- a/Engine/Source/Runtime/ApplicationCore/Private/Mac/CocoaWindow.cpp
+++ b/Engine/Source/Runtime/ApplicationCore/Private/Mac/CocoaWindow.cpp
@@ -556,7 +556,7 @@ NSString* NSPerformDragOperation = @"NSPerformDragOperation";
 
 - (void)UpdateAccessibilityView:(AccessibleWidgetId) InAccessibilityWindowId
 {
-	checkf(!IsInGameThread(), TEXT("Updating accessibility view in FCocoaWindow from Game Thread! Accessibility  can only be done on Main Thread!"));
+	checkf([NSThread isMainThread], TEXT("Updating accessibility view in FCocoaWindow from wrong Thread! Accessibility can only be done on Main Thread!"));
 	NSView* OpenGLView = [self openGLView];
 	// FCocoaAccessibilityView is the base class all custom NSViews must inherit from to support accessibility
 	if([OpenGLView isKindOfClass:[FCocoaAccessibilityView class]])
@@ -568,7 +568,7 @@ NSString* NSPerformDragOperation = @"NSPerformDragOperation";
 
 - (void)ClearAccessibilityView
 {
-	checkf(!IsInGameThread(), TEXT("Updating accessibility view in FCocoaWindow from Game Thread! Accessibility  can only be done on Main Thread!"));
+	checkf([NSThread isMainThread], TEXT("Updating accessibility view in FCocoaWindow from wrong Thread! Accessibility can only be done on Main Thread!"));
 	NSView* OpenGLView = [self openGLView];
 	// FCocoaAccessibilityView is the base class all custom NSViews must inherit from to support accessibility
 	if([OpenGLView isKindOfClass:[FCocoaAccessibilityView class]])
diff --git a/Engine/Source/Runtime/ApplicationCore/Private/Mac/MacProgramDelegate.cpp b/Engine/Source/Runtime/ApplicationCore/Private/Mac/MacProgramDelegate.cpp
new file mode 100644
--- /dev/null
+++ b/Engine/Source/Runtime/ApplicationCore/Private/Mac/MacProgramDelegate.cpp
@@ -0,0 +1,132 @@
+// Copyright Epic Games, Inc. All Rights Reserved.
+
+#include "Mac/MacProgramDelegate.h"
+#include "CoreMinimal.h"
+#include "HAL/ExceptionHandling.h"
+#include "Misc/CommandLine.h"
+#include "Mac/CocoaThread.h"
+
+@implementation MacProgramDelegate
+
+-(id)initWithProgramMain:(ProgramMainType)programMain programExit:(ProgramExitType)programExit
+{
+	[self init];
+	ProgramMain = programMain;
+	ProgramExit = programExit;
+	return self;
+}
+
+//handler for the quit apple event used by the Dock menu
+- (void)handleQuitEvent:(NSAppleEventDescriptor*)Event withReplyEvent:(NSAppleEventDescriptor*)ReplyEvent
+{
+	[self requestQuit:self];
+}
+
+- (IBAction)requestQuit:(id)Sender
+{
+	RequestEngineExit(TEXT("requestQuit"));
+}
+
+- (void)runGameThread:(id)Arg
+{
+	FPlatformMisc::SetGracefulTerminationHandler();
+	FPlatformMisc::SetCrashHandler(nullptr);
+	
+#if !UE_BUILD_SHIPPING
+	if (FParse::Param(*SavedCommandLine,TEXT("crashreports")))
+	{
+		GAlwaysReportCrash = true;
+	}
+#endif
+		
+#if UE_BUILD_DEBUG
+	if (!GAlwaysReportCrash)
+#else
+		if (FPlatformMisc::IsDebuggerPresent() && !GAlwaysReportCrash)
+#endif
+		{
+			ProgramMain(*SavedCommandLine);
+		}
+		else
+		{
+			GIsGuarded = 1;
+			ProgramMain(*SavedCommandLine);
+			GIsGuarded = 0;
+		}
+
+	ProgramExit();
+	
+	[NSApp terminate: self];
+}
+
+- (void) applicationWillTerminate:(NSNotification*)notification
+{
+	FTaskTagScope::SetTagStaticInit();
+}
+
+- (NSApplicationTerminateReply)applicationShouldTerminate:(NSApplication *)Sender;
+{
+	if(!IsEngineExitRequested() || ([NSThread gameThread] && [NSThread gameThread] != [NSThread mainThread]))
+	{
+		[self requestQuit:self];
+		return NSTerminateLater;
+	}
+	else
+	{
+		return NSTerminateNow;
+	}
+}
+
+- (void)applicationDidFinishLaunching:(NSNotification *)Notification
+{
+	//install the custom quit event handler
+	NSAppleEventManager* appleEventManager = [NSAppleEventManager sharedAppleEventManager];
+	[appleEventManager setEventHandler:self andSelector:@selector(handleQuitEvent:withReplyEvent:) forEventClass:kCoreEventClass andEventID:kAEQuitApplication];
+	
+	RunGameThread(self, @selector(runGameThread:));
+}
+
+
+
++(int)mainWithArgc:(int)argc argv:(char*[])argv programMain:(ProgramMainType)programMain
+	programExit:(ProgramExitType)programExit
+{
+	MacProgramDelegate* StandardDelegate = [[MacProgramDelegate alloc] initWithProgramMain:programMain programExit:programExit];
+	return [MacProgramDelegate mainWithArgc:argc argv:argv existingDelegate:StandardDelegate];
+}
+
++(int)mainWithArgc:(int)argc argv:(char*[])argv existingDelegate:(MacProgramDelegate*)delegate
+{
+	FTaskTagScope::SetTagNone();
+
+	for (int32 Option = 1; Option < argc; Option++)
+	{
+		delegate->SavedCommandLine += TEXT(" ");
+		FString Argument(ANSI_TO_TCHAR(argv[Option]));
+		if (Argument.Contains(TEXT(" ")))
+		{
+			if (Argument.Contains(TEXT("=")))
+			{
+				FString ArgName;
+				FString ArgValue;
+				Argument.Split( TEXT("="), &ArgName, &ArgValue );
+				Argument = FString::Printf( TEXT("%s=\"%s\""), *ArgName, *ArgValue );
+			}
+			else
+			{
+				Argument = FString::Printf(TEXT("\"%s\""), *Argument);
+			}
+		}
+		delegate->SavedCommandLine += Argument;
+	}
+
+	SCOPED_AUTORELEASE_POOL;
+	[NSApplication sharedApplication];
+	[NSApp setDelegate:delegate];
+	[NSApp run];
+	return 0;
+
+}
+
+@end
+
diff --git a/Engine/Source/Runtime/ApplicationCore/Public/Mac/MacProgramDelegate.h b/Engine/Source/Runtime/ApplicationCore/Public/Mac/MacProgramDelegate.h
new file mode 100644
--- /dev/null
+++ b/Engine/Source/Runtime/ApplicationCore/Public/Mac/MacProgramDelegate.h
@@ -0,0 +1,29 @@
+// Copyright Epic Games, Inc. All Rights Reserved.
+
+#pragma once
+
+#include "CoreMinimal.h"
+
+using ProgramMainType=int32(*)(const TCHAR *);
+using ProgramExitType=void(*)();
+
+@interface MacProgramDelegate : NSObject<NSApplicationDelegate, NSFileManagerDelegate>
+{
+	// the main function for the program (shared with other platforms)
+	ProgramMainType ProgramMain;
+	
+	// the function to run at exit time, on the game thread (very probagly FEngineLoop::Exit)
+	ProgramExitType ProgramExit;
+	
+	// the commandline made from argv, but a subclass could modify it
+	FString SavedCommandLine;
+}
+
+
+
+-(id)initWithProgramMain:(ProgramMainType)programMain programExit:(ProgramExitType)programExit;
+
++(int)mainWithArgc:(int)argc argv:(char*[])argv programMain:(ProgramMainType)programMain programExit:(ProgramExitType)programExit;
++(int)mainWithArgc:(int)argc argv:(char*[])argv existingDelegate:(MacProgramDelegate*)delegate;
+
+@end
diff --git a/Engine/Source/Runtime/Slate/Private/Framework/MultiBox/Mac/MacMenu.cpp b/Engine/Source/Runtime/Slate/Private/Framework/MultiBox/Mac/MacMenu.cpp
--- a/Engine/Source/Runtime/Slate/Private/Framework/MultiBox/Mac/MacMenu.cpp
+++ b/Engine/Source/Runtime/Slate/Private/Framework/MultiBox/Mac/MacMenu.cpp
@@ -284,7 +284,17 @@ namespace MacMenuHelper
 	}
 	
 	bool GMacPostInitStartupRequested = false;
-	bool GMacPostInitStartUpComplete = false; 
+	bool GMacPostInitStartUpComplete = false;
+
+	void ConditionalDelayPostStartup()
+	{
+		// In case an obsecure app startup sequence has not managed to finish the menu startup correctly
+		// However only do this if the post init has been called
+		if(!MacMenuHelper::GMacPostInitStartUpComplete && MacMenuHelper::GMacPostInitStartupRequested)
+		{
+			FSlateMacMenu::PostInitStartup();
+		}
+	}
 };
 
 // Bind all low-level Application hooks that require to access this high-level MacMenu system which includes NSApp Menu's and slate menus
@@ -443,13 +453,8 @@ void FSlateMacMenu::LanguageChanged()
 
 void FSlateMacMenu::UpdateApplicationMenu(bool bMacApplicationModalMode)
 {
-	// In case an obsecure app startup sequence has not managed to finish the menu startup correctly
-	// However only do this if the post init has been called
-	if(!MacMenuHelper::GMacPostInitStartUpComplete && MacMenuHelper::GMacPostInitStartupRequested)
-	{
-		PostInitStartup();
-	}
-
+	MacMenuHelper::ConditionalDelayPostStartup();
+	
     NSMenu* MainMenu = [NSApp mainMenu];
     NSMenuItem* AppMenuItem = [MainMenu itemWithTitle:@"AppMenuItem"];
     NSMenu* AppMenu = [AppMenuItem submenu];
@@ -521,6 +526,8 @@ void FSlateMacMenu::UpdateWindowMenu(bool bMacApplicationModalMode)
 
 void FSlateMacMenu::UpdateWithMultiBox(const TSharedPtr< FMultiBox > MultiBox)
 {
+	MacMenuHelper::ConditionalDelayPostStartup();
+	
 	// The dispatch block can't handle TSharedPtr correctly, so we use a small trick to pass MultiBox safely
 	struct FSafeMultiBoxPass
 	{
