[Insights] Session Browser: Fixed utrace file delete confirmation message box not showing and freezing UI.

diff --git a/Engine/Source/Programs/UnrealInsights/Resources/Windows/UnrealInsights.rc b/Engine/Source/Programs/UnrealInsights/Resources/Windows/UnrealInsights.rc
--- a/Engine/Source/Programs/UnrealInsights/Resources/Windows/UnrealInsights.rc
+++ b/Engine/Source/Programs/UnrealInsights/Resources/Windows/UnrealInsights.rc
@@ -7,7 +7,7 @@
 //
 #include <windows.h>
 #include "Runtime/Launch/Resources/Version.h"
-#include "Runtime/Launch/Resources/Windows/resource.h"
+#include "Runtime/Launch/Resources/Windows/PCLaunch.rc"
 #include "Developer/TraceInsights/Public/Insights/Version.h"
 
 #define VER_FILEVERSION         UNREAL_INSIGHTS_VERSION_MAJOR,UNREAL_INSIGHTS_VERSION_MINOR,UNREAL_INSIGHTS_VERSION_PATCH,0
