[Insights] Fixed parsing for -Store=<address/ip>:<port> cmd line parameter.

diff --git a/Engine/Source/Programs/UnrealInsights/Private/UserInterfaceCommand.cpp b/Engine/Source/Programs/UnrealInsights/Private/UserInterfaceCommand.cpp
--- a/Engine/Source/Programs/UnrealInsights/Private/UserInterfaceCommand.cpp
+++ b/Engine/Source/Programs/UnrealInsights/Private/UserInterfaceCommand.cpp
@@ -274,24 +274,21 @@ void FUserInterfaceCommand::InitializeSlateApplication(bool bOpenTraceFile, cons
 		TraceId = FCString::Strtoi(*TraceIdString, &End, 10);
 	}
 
-	const uint32 MaxPath = FPlatformMisc::GetMaxPathLength();
-	TCHAR* StoreHost = new TCHAR[MaxPath + 1];
-	FCString::Strcpy(StoreHost, MaxPath, TEXT("127.0.0.1"));
+	FString StoreHost = TEXT("127.0.0.1");
 	uint32 StorePort = 0;
 	bool bUseCustomStoreAddress = false;
 
-	if (FParse::Value(FCommandLine::Get(), TEXT("-Store="), StoreHost, MaxPath, true))
+	if (FParse::Value(FCommandLine::Get(), TEXT("-Store="), StoreHost, true))
 	{
-		TCHAR* Port = FCString::Strchr(StoreHost, TEXT(':'));
-		if (Port)
+		int32 Index = INDEX_NONE;
+		if (StoreHost.FindChar(TEXT(':'), Index))
 		{
-			*Port = 0;
-			Port++;
-			StorePort = FCString::Atoi(Port);
+			StorePort = FCString::Atoi(*StoreHost.RightChop(Index + 1));
+			StoreHost.LeftInline(Index);
 		}
 		bUseCustomStoreAddress = true;
 	}
-	if (FParse::Value(FCommandLine::Get(), TEXT("-StoreHost="), StoreHost, MaxPath, true))
+	if (FParse::Value(FCommandLine::Get(), TEXT("-StoreHost="), StoreHost, true))
 	{
 		bUseCustomStoreAddress = true;
 	}
@@ -312,9 +309,9 @@ void FUserInterfaceCommand::InitializeSlateApplication(bool bOpenTraceFile, cons
 
 	if (bUseTraceId || bOpenTraceFile) // viewer mode
 	{
-		TCHAR Cmd[1024];
+		FString Cmd;
 		bool bExecuteCommand = false;
-		if (FParse::Value(FCommandLine::Get(), TEXT("-ExecOnAnalysisCompleteCmd="), Cmd, 1024, false))
+		if (FParse::Value(FCommandLine::Get(), TEXT("-ExecOnAnalysisCompleteCmd="), Cmd, false))
 		{
 			bExecuteCommand = true;
 		}
@@ -331,7 +328,7 @@ void FUserInterfaceCommand::InitializeSlateApplication(bool bOpenTraceFile, cons
 
 		if (bUseTraceId)
 		{
-			TraceInsightsModule.ConnectToStore(StoreHost, StorePort);
+			TraceInsightsModule.ConnectToStore(*StoreHost, StorePort);
 			TraceInsightsModule.StartAnalysisForTrace(TraceId, bAutoQuit);
 		}
 		else
@@ -341,17 +338,16 @@ void FUserInterfaceCommand::InitializeSlateApplication(bool bOpenTraceFile, cons
 	}
 	else // browser mode
 	{
-
-		TCHAR Cmd[1024];
+		FString Cmd;
 		bool bExecuteCommand = false;
-		if (FParse::Value(FCommandLine::Get(), TEXT("-ExecBrowserAutomationTest="), Cmd, 1024, false))
+		if (FParse::Value(FCommandLine::Get(), TEXT("-ExecBrowserAutomationTest="), Cmd, false))
 		{
 			bExecuteCommand = true;
 		}
 
 		if (bUseCustomStoreAddress)
 		{
-			TraceInsightsModule.ConnectToStore(StoreHost, StorePort);
+			TraceInsightsModule.ConnectToStore(*StoreHost, StorePort);
 		}
 		else
 		{
@@ -369,8 +365,6 @@ void FUserInterfaceCommand::InitializeSlateApplication(bool bOpenTraceFile, cons
 			TraceInsightsModule.RunAutomationTest(Cmd);
 		}
 	}
-
-	delete[] StoreHost;
 }
 
 ////////////////////////////////////////////////////////////////////////////////////////////////////
