[Insights] Enable Session Frontend in Insights Session Browser

diff --git a/Engine/Plugins/Messaging/TcpMessaging/TcpMessaging.uplugin b/Engine/Plugins/Messaging/TcpMessaging/TcpMessaging.uplugin
--- a/Engine/Plugins/Messaging/TcpMessaging/TcpMessaging.uplugin
+++ b/Engine/Plugins/Messaging/TcpMessaging/TcpMessaging.uplugin
@@ -14,7 +14,7 @@
 	"CanContainContent" : false,
 	"IsBetaVersion" : false,
 	"Installed" : false,
-	"SupportedPrograms": [ "UnrealFrontend" ],
+	"SupportedPrograms": [ "UnrealFrontend" , "UnrealInsights" ],
 	"Modules" :
 	[
 		{
@@ -24,7 +24,7 @@
 			"PlatformDenyList" :
 			[
 			],
-			"ProgramAllowList" : [ "UnrealFrontend" ]
+			"ProgramAllowList": [ "UnrealFrontend", "UnrealInsights" ]
 	    }
 	]
 }
diff --git a/Engine/Plugins/Messaging/UdpMessaging/UdpMessaging.uplugin b/Engine/Plugins/Messaging/UdpMessaging/UdpMessaging.uplugin
--- a/Engine/Plugins/Messaging/UdpMessaging/UdpMessaging.uplugin
+++ b/Engine/Plugins/Messaging/UdpMessaging/UdpMessaging.uplugin
@@ -15,6 +15,7 @@
 	"IsBetaVersion" : false,
 	"Installed" : false,
 	"SupportedPrograms": [
+		"UnrealInsights",
 		"UnrealFrontend",
 		"UnrealMultiUserServer",
 		"UnrealMultiUserSlateServer",
@@ -34,6 +35,7 @@
 			[
 			],
 			"ProgramAllowList": [
+				"UnrealInsights",
 				"UnrealFrontend",
 				"UnrealMultiUserServer",
 				"UnrealMultiUserSlateServer",
diff --git a/Engine/Programs/UnrealInsights/Config/DefaultEngine.ini b/Engine/Programs/UnrealInsights/Config/DefaultEngine.ini
--- a/Engine/Programs/UnrealInsights/Config/DefaultEngine.ini
+++ b/Engine/Programs/UnrealInsights/Config/DefaultEngine.ini
@@ -3,6 +3,8 @@
 +ProgramEnabledPlugins=TraceDataFilters
 +ProgramEnabledPlugins=SlateInsights
 +ProgramEnabledPlugins=RenderGraphInsights
++ProgramEnabledPlugins=UdpMessaging
++ProgramEnabledPlugins=TcpMessaging
 +ProgramEnabledPlugins=NetworkPredictionInsights
 +ProgramEnabledPlugins=VisualStudioSourceCodeAccess
 +ProgramEnabledPlugins=XcodeSourceCodeAccess
diff --git a/Engine/Source/Developer/TraceInsights/Private/Insights/TraceInsightsModule.cpp b/Engine/Source/Developer/TraceInsights/Private/Insights/TraceInsightsModule.cpp
--- a/Engine/Source/Developer/TraceInsights/Private/Insights/TraceInsightsModule.cpp
+++ b/Engine/Source/Developer/TraceInsights/Private/Insights/TraceInsightsModule.cpp
@@ -310,6 +310,7 @@ void FTraceInsightsModule::CreateSessionBrowser(const FCreateSessionBrowserParam
 			FTabManager::NewStack()
 			->AddTab(FInsightsManagerTabs::TraceStoreTabId, ETabState::OpenedTab)
 			->AddTab(FInsightsManagerTabs::ConnectionTabId, ETabState::OpenedTab)
+			->AddTab(FName("SessionFrontend"), ETabState::OpenedTab)
 			//->AddTab(FInsightsManagerTabs::LauncherTabId, ETabState::ClosedTab)
 			->SetForegroundTab(FInsightsManagerTabs::TraceStoreTabId)
 		)
diff --git a/Engine/Source/Programs/UnrealInsights/Private/UserInterfaceCommand.cpp b/Engine/Source/Programs/UnrealInsights/Private/UserInterfaceCommand.cpp
--- a/Engine/Source/Programs/UnrealInsights/Private/UserInterfaceCommand.cpp
+++ b/Engine/Source/Programs/UnrealInsights/Private/UserInterfaceCommand.cpp
@@ -179,6 +179,8 @@ void FUserInterfaceCommand::Run()
 	delete[] TraceFile;
 	TraceFile = nullptr;
 
+	IPluginManager::Get().LoadModulesForEnabledPlugins(ELoadingPhase::PostDefault);
+
 	// Initialize source code access.
 	// Load the source code access module.
 	ISourceCodeAccessModule& SourceCodeAccessModule = FModuleManager::LoadModuleChecked<ISourceCodeAccessModule>(FName("SourceCodeAccess"));
@@ -335,6 +337,18 @@ void FUserInterfaceCommand::InitializeSlateApplication(bool bOpenTraceFile, cons
 	}
 	else // browser mode
 	{
+		// ensure target platform manager is referenced early as it must be created on the main thread
+		FModuleManager::Get().LoadModuleChecked("DesktopPlatform");
+		FConfigCacheIni::InitializeConfigSystem();
+		GetTargetPlatformManager();
+
+		FModuleManager::Get().LoadModuleChecked("Messaging");
+		FModuleManager::Get().LoadModuleChecked("OutputLog");
+
+		// load optional modules
+		FModuleManager::Get().LoadModule("DeviceManager");
+		FModuleManager::Get().LoadModule("SessionFrontend");
+
 		FString Cmd;
 		bool bExecuteCommand = false;
 		if (FParse::Value(FCommandLine::Get(), TEXT("-ExecBrowserAutomationTest="), Cmd, false))
diff --git a/Engine/Source/Programs/UnrealInsights/UnrealInsights.Build.cs b/Engine/Source/Programs/UnrealInsights/UnrealInsights.Build.cs
--- a/Engine/Source/Programs/UnrealInsights/UnrealInsights.Build.cs
+++ b/Engine/Source/Programs/UnrealInsights/UnrealInsights.Build.cs
@@ -26,6 +26,17 @@ public class UnrealInsights : ModuleRules
 			}
 		);
 
+		// For Session Frontend and Message Bus
+		PrivateDependencyModuleNames.AddRange(
+			new string[] {
+				"OutputLog",
+				"SessionFrontend",
+				"TargetDeviceServices",
+				"TcpMessaging",
+				"UdpMessaging",
+			}
+		);
+
 		PrivateIncludePathModuleNames.AddRange(
 			new string[] {
 				"SlateReflector"
