[Insights] Added DisableFrameThrottle flag

diff --git a/Engine/Source/Developer/TraceInsights/Private/Insights/TraceInsightsModule.cpp b/Engine/Source/Developer/TraceInsights/Private/Insights/TraceInsightsModule.cpp
--- a/Engine/Source/Developer/TraceInsights/Private/Insights/TraceInsightsModule.cpp
+++ b/Engine/Source/Developer/TraceInsights/Private/Insights/TraceInsightsModule.cpp
@@ -327,6 +327,7 @@ void FTraceInsightsModule::CreateSessionBrowser(const FCreateSessionBrowserParam
 		TraceStoreWnd->SetEnableAutomaticTesting(Params.bInitializeTesting);
 		TraceStoreWnd->SetEnableDebugTools(Params.bAllowDebugTools);
 		TraceStoreWnd->SetStartProcessWithStompMalloc(Params.bStartProcessWithStompMalloc);
+		TraceStoreWnd->SetDisableFramerateThrottle(Params.bDisableFramerateThrottle);
 	}
 }
 
diff --git a/Engine/Source/Developer/TraceInsights/Private/Insights/Widgets/SStartPageWindow.cpp b/Engine/Source/Developer/TraceInsights/Private/Insights/Widgets/SStartPageWindow.cpp
--- a/Engine/Source/Developer/TraceInsights/Private/Insights/Widgets/SStartPageWindow.cpp
+++ b/Engine/Source/Developer/TraceInsights/Private/Insights/Widgets/SStartPageWindow.cpp
@@ -3823,6 +3823,10 @@ void STraceStoreWindow::GetExtraCommandLineParams(FString& OutParams) const
 	{
 		OutParams.Append(TEXT(" -stompmalloc"));
 	}
+	if (bDisableFramerateThrottle)
+	{
+		OutParams.Append(TEXT(" -DisableFramerateThrottle"));
+	}
 }
 
 ////////////////////////////////////////////////////////////////////////////////////////////////////
diff --git a/Engine/Source/Developer/TraceInsights/Private/Insights/Widgets/SStartPageWindow.h b/Engine/Source/Developer/TraceInsights/Private/Insights/Widgets/SStartPageWindow.h
--- a/Engine/Source/Developer/TraceInsights/Private/Insights/Widgets/SStartPageWindow.h
+++ b/Engine/Source/Developer/TraceInsights/Private/Insights/Widgets/SStartPageWindow.h
@@ -355,6 +355,9 @@ public:
 	void SetStartProcessWithStompMalloc(bool InValue) { bStartProcessWithStompMalloc = InValue; };
 	bool GetStartProcessWithStompMalloc() const { return bStartProcessWithStompMalloc; };
 
+	void SetDisableFramerateThrottle(bool InValue) { bDisableFramerateThrottle = InValue; };
+	bool GetDisableFramerateThrottle() const { return bDisableFramerateThrottle; };
+
 	void OnFilterChanged();
 	const TArray<TSharedPtr<FTraceViewModel>>& GetAllAvailableTraces() const;
 	void SetDeleteTraceConfirmationWindowVisibility(bool bIsVisibile);
@@ -698,6 +701,7 @@ private:
 	bool bEnableAutomaticTesting = false;
 	bool bEnableDebugTools = false;
 	bool bStartProcessWithStompMalloc = false;
+	bool bDisableFramerateThrottle = false;
 
 	bool bSetKeyboardFocusOnNextTick = false;
 
diff --git a/Engine/Source/Developer/TraceInsights/Public/Insights/IUnrealInsightsModule.h b/Engine/Source/Developer/TraceInsights/Public/Insights/IUnrealInsightsModule.h
--- a/Engine/Source/Developer/TraceInsights/Public/Insights/IUnrealInsightsModule.h
+++ b/Engine/Source/Developer/TraceInsights/Public/Insights/IUnrealInsightsModule.h
@@ -144,6 +144,7 @@ struct TRACEINSIGHTS_API FCreateSessionBrowserParams
 	bool bAllowDebugTools = false;
 	bool bInitializeTesting = false;
 	bool bStartProcessWithStompMalloc = false;
+	bool bDisableFramerateThrottle = false;
 };
 
 ////////////////////////////////////////////////////////////////////////////////////////////////////
diff --git a/Engine/Source/Programs/UnrealInsights/Private/UserInterfaceCommand.cpp b/Engine/Source/Programs/UnrealInsights/Private/UserInterfaceCommand.cpp
--- a/Engine/Source/Programs/UnrealInsights/Private/UserInterfaceCommand.cpp
+++ b/Engine/Source/Programs/UnrealInsights/Private/UserInterfaceCommand.cpp
@@ -199,6 +199,8 @@ void FUserInterfaceCommand::Run()
 	SharedPointerTesting::TestSharedPointer<ESPMode::ThreadSafe>();
 #endif
 
+	const bool bDisableFramerateThrottle = FParse::Param(FCommandLine::Get(), TEXT("DisableFramerateThrottle"));
+
 	// Enter main loop.
 	double DeltaTime = 0.0;
 	double LastTime = FPlatformTime::Seconds();
@@ -217,7 +219,9 @@ void FUserInterfaceCommand::Run()
 		FTSTicker::GetCoreTicker().Tick(static_cast<float>(DeltaTime));
 
 		// Throttle frame rate.
-		const float FrameTime = UserInterfaceCommand::IsApplicationBackground() ? BackgroundFrameTime : IdealFrameTime;
+		const float FrameTime = bDisableFramerateThrottle ? IdealFrameTime
+														  : UserInterfaceCommand::IsApplicationBackground() ? BackgroundFrameTime : IdealFrameTime;
+
 		UserInterfaceCommand::AdaptiveSleep(FMath::Max<float>(0.0f, FrameTime - static_cast<float>(FPlatformTime::Seconds() - LastTime)));
 
 		double CurrentTime = FPlatformTime::Seconds();
@@ -369,6 +373,7 @@ void FUserInterfaceCommand::InitializeSlateApplication(bool bOpenTraceFile, cons
 		Params.bAllowDebugTools = bAllowDebugTools;
 		Params.bInitializeTesting = bInitializeTesting;
 		Params.bStartProcessWithStompMalloc = FParse::Param(FCommandLine::Get(), TEXT("stompmalloc"));
+		Params.bDisableFramerateThrottle = FParse::Param(FCommandLine::Get(), TEXT("DisableFramerateThrottle"));
 		TraceInsightsModule.CreateSessionBrowser(Params);
 
 		if (bExecuteCommand)
