[Insights]

diff --git a/Engine/Source/Developer/TraceInsights/Private/Insights/Tests/InsightsTestRunner.cpp b/Engine/Source/Developer/TraceInsights/Private/Insights/Tests/InsightsTestRunner.cpp
--- a/Engine/Source/Developer/TraceInsights/Private/Insights/Tests/InsightsTestRunner.cpp
+++ b/Engine/Source/Developer/TraceInsights/Private/Insights/Tests/InsightsTestRunner.cpp
@@ -4,16 +4,19 @@
 
 #if !UE_BUILD_SHIPPING && !WITH_EDITOR
 
-#include "Misc/App.h"
-#include "Misc/AutomationTest.h"
 #include "IAutomationControllerModule.h"
 #include "IAutomationWindowModule.h"
 #include "IAutomationWorkerModule.h"
 #include "ISessionManager.h"
 #include "ISessionServicesModule.h"
+#include "Misc/App.h"
+#include "Misc/AutomationTest.h"
 #include "Misc/CoreMisc.h"
 #include "Widgets/Docking/SDockTab.h"
 
+// TraceInsightsCore
+#include "InsightsCore/Common/InsightsCoreStyle.h"
+
 // TraceInsights
 #include "Insights/Common/InsightsMenuBuilder.h"
 #include "Insights/InsightsManager.h"
@@ -68,6 +71,7 @@ void FInsightsTestRunner::Initialize(IUnrealInsightsModule& InsightsModule)
 		FModuleManager::Get().LoadModule("AutomationWorker");
 	}
 
+	using namespace UE::Insights;
 	SessionAnalysisCompletedHandle = FInsightsManager::Get()->GetSessionAnalysisCompletedEvent().AddSP(this, &FInsightsTestRunner::OnSessionAnalysisCompleted);
 
 	// Register tick functions.
@@ -86,6 +90,8 @@ void FInsightsTestRunner::Shutdown()
 
 void FInsightsTestRunner::RegisterMajorTabs(IUnrealInsightsModule& InsightsModule)
 {
+	using namespace UE::Insights;
+
 	if (bInitAutomationModules)
 	{
 		const FInsightsMajorTabConfig& AutomationConfig = InsightsModule.FindMajorTabConfig(FInsightsManagerTabs::AutomationWindowTabId);
@@ -96,7 +102,7 @@ void FInsightsTestRunner::RegisterMajorTabs(IUnrealInsightsModule& InsightsModul
 				FOnSpawnTab::CreateRaw(this, &FInsightsTestRunner::SpawnAutomationWindowTab))
 				.SetDisplayName(AutomationConfig.TabLabel.IsSet() ? AutomationConfig.TabLabel.GetValue() : LOCTEXT("AutomationTab", "Automation"))
 				.SetTooltipText(AutomationConfig.TabTooltip.IsSet() ? AutomationConfig.TabTooltip.GetValue() : LOCTEXT("AutomationTooltipText", "Opens the automation tab."))
-				.SetIcon(AutomationConfig.TabIcon.IsSet() ? AutomationConfig.TabIcon.GetValue() : FSlateIcon(FInsightsStyle::GetStyleSetName(), "Icons.TestAutomation"));
+				.SetIcon(AutomationConfig.TabIcon.IsSet() ? AutomationConfig.TabIcon.GetValue() : FSlateIcon(FInsightsCoreStyle::GetStyleSetName(), "Icons.TestAutomation"));
 
 			TSharedRef<FWorkspaceItem> Group = AutomationConfig.WorkspaceGroup.IsValid() ? AutomationConfig.WorkspaceGroup.ToSharedRef() : FInsightsManager::Get()->GetInsightsMenuBuilder()->GetWindowsGroup();
 			TabSpawnerEntry.SetGroup(Group);
@@ -108,6 +114,7 @@ void FInsightsTestRunner::RegisterMajorTabs(IUnrealInsightsModule& InsightsModul
 
 void FInsightsTestRunner::UnregisterMajorTabs()
 {
+	using namespace UE::Insights;
 	FGlobalTabmanager::Get()->UnregisterNomadTabSpawner(FInsightsManagerTabs::AutomationWindowTabId);
 }
 
@@ -169,6 +176,7 @@ void FInsightsTestRunner::OnSessionAnalysisCompleted()
 
 	bIsAnalysisComplete = true;
 
+	using namespace UE::Insights;
 	FInsightsManager::Get()->GetSessionAnalysisCompletedEvent().Remove(SessionAnalysisCompletedHandle);
 }
 
diff --git a/Engine/Source/Developer/TraceInsights/Private/Insights/Tests/InsightsTestUtils.cpp b/Engine/Source/Developer/TraceInsights/Private/Insights/Tests/InsightsTestUtils.cpp
--- a/Engine/Source/Developer/TraceInsights/Private/Insights/Tests/InsightsTestUtils.cpp
+++ b/Engine/Source/Developer/TraceInsights/Private/Insights/Tests/InsightsTestUtils.cpp
@@ -288,6 +288,7 @@ bool FInsightsTestUtils::KillUTS(double Timeout) const
 
 void FInsightsTestUtils::ResetSession() const
 {
+	using namespace UE::Insights;
 	TSharedPtr<FInsightsManager> InsightsManager = FInsightsManager::Get();
 	if (InsightsManager.IsValid())
 	{
diff --git a/Engine/Source/Developer/TraceInsights/Private/Insights/Tests/UnitTests.cpp b/Engine/Source/Developer/TraceInsights/Private/Insights/Tests/UnitTests.cpp
--- a/Engine/Source/Developer/TraceInsights/Private/Insights/Tests/UnitTests.cpp
+++ b/Engine/Source/Developer/TraceInsights/Private/Insights/Tests/UnitTests.cpp
@@ -1,5 +1,6 @@
 // Copyright Epic Games, Inc. All Rights Reserved.
-#include "CoreMinimal.h"
+
+#include "CoreTypes.h"
 #include "Misc/AutomationTest.h"
 
 IMPLEMENT_SIMPLE_AUTOMATION_TEST(FTraceInsightsUnitTest, "System.Insights.UnitTest", EAutomationTestFlags_ApplicationContextMask | EAutomationTestFlags::SmokeFilter)
@@ -7,4 +8,4 @@ IMPLEMENT_SIMPLE_AUTOMATION_TEST(FTraceInsightsUnitTest, "System.Insights.UnitTe
 bool FTraceInsightsUnitTest::RunTest(const FString& Parameters)
 {
 	return true;
-}
\ No newline at end of file
+}
diff --git a/Engine/Source/Developer/TraceInsights/Private/Insights/TimingProfiler/Widgets/SStatsTableCell.cpp b/Engine/Source/Developer/TraceInsights/Private/Insights/TimingProfiler/Widgets/SStatsTableCell.cpp
--- a/Engine/Source/Developer/TraceInsights/Private/Insights/TimingProfiler/Widgets/SStatsTableCell.cpp
+++ b/Engine/Source/Developer/TraceInsights/Private/Insights/TimingProfiler/Widgets/SStatsTableCell.cpp
@@ -99,7 +99,7 @@ TSharedRef<SWidget> SStatsTableCell::GenerateWidgetForNameColumn(const FArgument
 			.HeightOverride(14.0f)
 			[
 				SNew(SBorder)
-				.BorderImage(FInsightsStyle::Get().GetBrush("WhiteBrush"))
+				.BorderImage(FAppStyle::Get().GetBrush("WhiteBrush"))
 				.BorderBackgroundColor(this, &SStatsTableCell::GetBoxColorAndOpacity)
 				.HAlign(HAlign_Fill)
 				.VAlign(VAlign_Fill)
diff --git a/Engine/Source/Developer/TraceInsights/Private/Insights/TimingProfiler/Widgets/STimingProfilerWindow.h b/Engine/Source/Developer/TraceInsights/Private/Insights/TimingProfiler/Widgets/STimingProfilerWindow.h
--- a/Engine/Source/Developer/TraceInsights/Private/Insights/TimingProfiler/Widgets/STimingProfilerWindow.h
+++ b/Engine/Source/Developer/TraceInsights/Private/Insights/TimingProfiler/Widgets/STimingProfilerWindow.h
@@ -8,7 +8,7 @@
 #include "Insights/ITimingViewSession.h" // for ETimeChangedFlags
 #include "Insights/Widgets/SMajorTabWindow.h"
 
-class SLogView;
+namespace UE::Insights { class SLogView; }
 
 namespace UE::Insights::TimingProfiler
 {
diff --git a/Engine/Source/Developer/TraceInsights/Private/Insights/TraceInsightsModule.cpp b/Engine/Source/Developer/TraceInsights/Private/Insights/TraceInsightsModule.cpp
--- a/Engine/Source/Developer/TraceInsights/Private/Insights/TraceInsightsModule.cpp
+++ b/Engine/Source/Developer/TraceInsights/Private/Insights/TraceInsightsModule.cpp
@@ -37,12 +37,14 @@
 #include "Insights/TaskGraphProfiler/TaskGraphProfilerManager.h"
 #include "Insights/Tests/InsightsTestRunner.h"
 #include "Insights/TimingProfiler/TimingProfilerManager.h"
-#include "Insights/Widgets/SStartPageWindow.h"
 
 ////////////////////////////////////////////////////////////////////////////////////////////////////
 
 DEFINE_LOG_CATEGORY(TraceInsights);
 
+namespace UE::Insights
+{
+
 IMPLEMENT_MODULE(FTraceInsightsModule, TraceInsights);
 
 FString FTraceInsightsModule::UnrealInsightsLayoutIni;
@@ -182,7 +184,7 @@ void FTraceInsightsModule::UnregisterComponent(TSharedPtr<IInsightsComponent> Co
 
 ////////////////////////////////////////////////////////////////////////////////////////////////////
 
-TSharedPtr<Insights::IInsightsManager> FTraceInsightsModule::GetInsightsManager()
+TSharedPtr<::Insights::IInsightsManager> FTraceInsightsModule::GetInsightsManager()
 {
 	return FInsightsManager::Get();
 }
@@ -256,97 +258,6 @@ void FTraceInsightsModule::UnregisterTabSpawners()
 
 ////////////////////////////////////////////////////////////////////////////////////////////////////
 
-void FTraceInsightsModule::CreateSessionBrowser(const FCreateSessionBrowserParams& Params)
-{
-	RegisterTabSpawners();
-
-	//////////////////////////////////////////////////
-	// Create the main window.
-
-	const bool bEmbedTitleAreaContent = false;
-
-	// Get desktop metrics. It also ensures the correct metrics will be used later in SWindow.
-	FDisplayMetrics DisplayMetrics;
-	FSlateApplication::Get().GetDisplayMetrics(DisplayMetrics);
-	const float DPIScaleFactor = FPlatformApplicationMisc::GetDPIScaleFactorAtPoint(
-		static_cast<float>(DisplayMetrics.PrimaryDisplayWorkAreaRect.Left),
-		static_cast<float>(DisplayMetrics.PrimaryDisplayWorkAreaRect.Top));
-
-	const FVector2D ClientSize(1280.0f * DPIScaleFactor, 720.0f * DPIScaleFactor);
-
-	TSharedRef<SWindow> RootWindow = SNew(SWindow)
-		.Title(NSLOCTEXT("TraceInsightsModule", "UnrealInsightsBrowserAppName", "Unreal Insights Frontend"))
-		.CreateTitleBar(!bEmbedTitleAreaContent)
-		.SupportsMaximize(true)
-		.SupportsMinimize(true)
-		.IsInitiallyMaximized(false)
-		.IsInitiallyMinimized(false)
-		.SizingRule(ESizingRule::UserSized)
-		.AutoCenter(EAutoCenter::PreferredWorkArea)
-		.ClientSize(ClientSize)
-		.AdjustInitialSizeAndPositionForDPIScale(false);
-
-	//RootWindow->GetTitleBar()->SetAllowMenuBar(true);
-
-	const bool bShowRootWindowImmediately = false;
-	FSlateApplication::Get().AddWindow(RootWindow, bShowRootWindowImmediately);
-
-	FGlobalTabmanager::Get()->SetRootWindow(RootWindow);
-	FGlobalTabmanager::Get()->SetAllowWindowMenuBar(true);
-
-	FSlateNotificationManager::Get().SetRootWindow(RootWindow);
-
-	//////////////////////////////////////////////////
-	// Setup the window's content.
-
-	TSharedRef<FTabManager::FLayout> DefaultLayout = FTabManager::NewLayout("TraceSessionBrowserLayout_v1.1");
-	DefaultLayout->AddArea
-	(
-		FTabManager::NewPrimaryArea()
-		->Split
-		(
-			FTabManager::NewStack()
-			->AddTab(FInsightsManagerTabs::TraceStoreTabId, ETabState::OpenedTab)
-			->AddTab(FInsightsManagerTabs::ConnectionTabId, ETabState::OpenedTab)
-			->AddTab(FName("SessionFrontend"), ETabState::OpenedTab)
-			//->AddTab(FInsightsManagerTabs::LauncherTabId, ETabState::ClosedTab)
-			->SetForegroundTab(FInsightsManagerTabs::TraceStoreTabId)
-		)
-	);
-
-	AddAreaForWidgetReflector(DefaultLayout, Params.bAllowDebugTools);
-
-	// Load layout from ini file.
-	PersistentLayout = FLayoutSaveRestore::LoadFromConfig(UnrealInsightsLayoutIni, DefaultLayout);
-
-	// Restore application layout.
-	const EOutputCanBeNullptr OutputCanBeNullptr = EOutputCanBeNullptr::Never;
-	TSharedPtr<SWidget> Content = FGlobalTabmanager::Get()->RestoreFrom(PersistentLayout.ToSharedRef(), RootWindow, bEmbedTitleAreaContent, OutputCanBeNullptr);
-	RootWindow->SetContent(Content.ToSharedRef());
-	RootWindow->GetOnWindowClosedEvent().AddRaw(this, &FTraceInsightsModule::OnWindowClosedEvent);
-
-	//////////////////////////////////////////////////
-	// Show the window.
-
-	RootWindow->ShowWindow();
-	const bool bForceWindowToFront = true;
-	RootWindow->BringToFront(bForceWindowToFront);
-
-	//////////////////////////////////////////////////
-	// Set up command line parameter forwarding.
-
-	TSharedPtr<STraceStoreWindow> TraceStoreWnd = FInsightsManager::Get()->GetTraceStoreWindow();
-	if (TraceStoreWnd.IsValid())
-	{
-		TraceStoreWnd->SetEnableAutomaticTesting(Params.bInitializeTesting);
-		TraceStoreWnd->SetEnableDebugTools(Params.bAllowDebugTools);
-		TraceStoreWnd->SetStartProcessWithStompMalloc(Params.bStartProcessWithStompMalloc);
-		TraceStoreWnd->SetDisableFramerateThrottle(Params.bDisableFramerateThrottle);
-	}
-}
-
-////////////////////////////////////////////////////////////////////////////////////////////////////
-
 void FTraceInsightsModule::CreateSessionViewer(bool bAllowDebugTools)
 {
 	RegisterTabSpawners();
@@ -629,24 +540,6 @@ void FTraceInsightsModule::ScheduleCommand(const FString& InCmd)
 
 ////////////////////////////////////////////////////////////////////////////////////////////////////
 
-void FTraceInsightsModule::RunAutomationTest(const FString& InCmd)
-{
-#if !UE_BUILD_SHIPPING && !WITH_EDITOR
-	FString ActualCmd = InCmd;
-	ActualCmd.TrimCharInline(TEXT('\"'), nullptr);
-	ActualCmd.TrimCharInline(TEXT('\''), nullptr);
-
-	if (ActualCmd.StartsWith(TEXT("Automation RunTests")))
-	{
-		FInsightsTestRunner::Get()->ScheduleCommand(ActualCmd);
-		FInsightsTestRunner::Get()->RunTests();
-		return;
-	}
-#endif
-}
-
-////////////////////////////////////////////////////////////////////////////////////////////////////
-
 bool FTraceInsightsModule::Exec(const TCHAR* Cmd, FOutputDevice& Ar)
 {
 	for (TSharedRef<IInsightsComponent>& Component : Components)
@@ -679,3 +572,5 @@ void FTraceInsightsModule::HandleCodeAccessorOpenFileFailed(const FString& Filen
 }
 
 ////////////////////////////////////////////////////////////////////////////////////////////////////
+
+} // namespace UE::Insights
diff --git a/Engine/Source/Developer/TraceInsights/Private/Insights/TraceInsightsModule.h b/Engine/Source/Developer/TraceInsights/Private/Insights/TraceInsightsModule.h
--- a/Engine/Source/Developer/TraceInsights/Private/Insights/TraceInsightsModule.h
+++ b/Engine/Source/Developer/TraceInsights/Private/Insights/TraceInsightsModule.h
@@ -25,6 +25,9 @@ class SDockTab;
 class FSpawnTabArgs;
 class SWindow;
 
+namespace UE::Insights
+{
+
 /**
  * Implements the Trace Insights module.
  */
@@ -44,13 +47,11 @@ public:
 	}
 
 	virtual void CreateDefaultStore() override;
-
 	FString GetDefaultStoreDir();
 
 	virtual UE::Trace::FStoreClient* GetStoreClient() override;
 	virtual bool ConnectToStore(const TCHAR* InStoreHost, uint32 InStorePort=0) override;
 
-	virtual void CreateSessionBrowser(const FCreateSessionBrowserParams& Params) override;
 	virtual void CreateSessionViewer(bool bAllowDebugTools = false) override;
 
 	virtual TSharedPtr<const TraceServices::IAnalysisSession> GetAnalysisSession() const override;
@@ -63,7 +64,7 @@ public:
 	virtual void RegisterComponent(TSharedPtr<IInsightsComponent> Component) override;
 	virtual void UnregisterComponent(TSharedPtr<IInsightsComponent> Component) override;
 
-	virtual TSharedPtr<Insights::IInsightsManager> GetInsightsManager() override;
+	virtual TSharedPtr<::Insights::IInsightsManager> GetInsightsManager() override;
 
 	virtual void RegisterMajorTabConfig(const FName& InMajorTabId, const FInsightsMajorTabConfig& InConfig) override;
 	virtual void UnregisterMajorTabConfig(const FName& InMajorTabId) override;
@@ -82,30 +83,15 @@ public:
 
 	virtual void InitializeTesting(bool InInitAutomationModules, bool InAutoQuit) override;
 	virtual void ScheduleCommand(const FString& InCmd) override;
-	virtual void RunAutomationTest(const FString& InCmd) override;
 	virtual bool Exec(const TCHAR* Cmd, FOutputDevice& Ar) override;
 
-protected:
-	void InitTraceStore();
-
+private:
 	void RegisterTabSpawners();
 	void UnregisterTabSpawners();
 
 	void AddAreaForSessionViewer(TSharedRef<FTabManager::FLayout> Layout);
 	void AddAreaForWidgetReflector(TSharedRef<FTabManager::FLayout> Layout, bool bAllowDebugTools);
 
-	/** Callback called when a major tab is closed. */
-	void OnTabBeingClosed(TSharedRef<SDockTab> TabBeingClosed);
-
-	/** Trace Store */
-	TSharedRef<SDockTab> SpawnTraceStoreTab(const FSpawnTabArgs& Args);
-
-	/** Connection */
-	TSharedRef<SDockTab> SpawnConnectionTab(const FSpawnTabArgs& Args);
-
-	/** Launcher */
-	TSharedRef<SDockTab> SpawnLauncherTab(const FSpawnTabArgs& Args);
-
 	/** Session Info */
 	TSharedRef<SDockTab> SpawnSessionInfoTab(const FSpawnTabArgs& Args);
 
@@ -115,6 +101,7 @@ protected:
 
 	void HandleCodeAccessorOpenFileFailed(const FString& Filename);
 
+private:
 	TSharedPtr<TraceServices::IAnalysisService> TraceAnalysisService;
 	TSharedPtr<TraceServices::IModuleService> TraceModuleService;
 
@@ -128,3 +115,5 @@ protected:
 
 	TArray<TSharedRef<IInsightsComponent>> Components;
 };
+
+} // namespace UE::Insights
diff --git a/Engine/Source/Developer/TraceInsights/Private/Insights/ViewModels/LogFilter.cpp b/Engine/Source/Developer/TraceInsights/Private/Insights/ViewModels/LogFilter.cpp
--- a/Engine/Source/Developer/TraceInsights/Private/Insights/ViewModels/LogFilter.cpp
+++ b/Engine/Source/Developer/TraceInsights/Private/Insights/ViewModels/LogFilter.cpp
@@ -8,6 +8,9 @@
 #include "Insights/Log.h"
 #include "Insights/Widgets/SLogView.h"
 
+namespace UE::Insights
+{
+
 ////////////////////////////////////////////////////////////////////////////////////////////////////
 /**
  * Expression context to test the given messages against the current text filter.
@@ -310,3 +313,5 @@ void FLogFilteringAsyncTask::DoWork()
 }
 
 ////////////////////////////////////////////////////////////////////////////////////////////////////
+
+} // namespace UE::Insights
diff --git a/Engine/Source/Developer/TraceInsights/Private/Insights/ViewModels/LogFilter.h b/Engine/Source/Developer/TraceInsights/Private/Insights/ViewModels/LogFilter.h
--- a/Engine/Source/Developer/TraceInsights/Private/Insights/ViewModels/LogFilter.h
+++ b/Engine/Source/Developer/TraceInsights/Private/Insights/ViewModels/LogFilter.h
@@ -15,6 +15,10 @@
 ////////////////////////////////////////////////////////////////////////////////////////////////////
 
 template<typename TTask> class FAsyncTask;
+
+namespace UE::Insights
+{
+
 class SLogView;
 
 ////////////////////////////////////////////////////////////////////////////////////////////////////
@@ -188,7 +192,7 @@ private:
 	/** A copy of the filter settings. */
 	FLogFilter Filter;
 
-	/** Shared pointer to parent LogView widget. Used for accesing the cache and to check if cancel is requested. */
+	/** Shared pointer to parent LogView widget. Used for accessing the cache and to check if cancel is requested. */
 	TSharedPtr<SLogView> LogView;
 
 	/** The output filtered messages. */
@@ -196,3 +200,5 @@ private:
 };
 
 ////////////////////////////////////////////////////////////////////////////////////////////////////
+
+} // namespace UE::Insights
diff --git a/Engine/Source/Developer/TraceInsights/Private/Insights/ViewModels/TimingGraphTrack.cpp b/Engine/Source/Developer/TraceInsights/Private/Insights/ViewModels/TimingGraphTrack.cpp
--- a/Engine/Source/Developer/TraceInsights/Private/Insights/ViewModels/TimingGraphTrack.cpp
+++ b/Engine/Source/Developer/TraceInsights/Private/Insights/ViewModels/TimingGraphTrack.cpp
@@ -310,7 +310,7 @@ void FTimingGraphTrack::Update(const ITimingTrackUpdateContext& Context)
 
 void FTimingGraphTrack::AddDefaultFrameSeries()
 {
-	const FInsightsSettings& Settings = FInsightsManager::Get()->GetSettings();
+	const FInsightsSettings& Settings = UE::Insights::FInsightsManager::Get()->GetSettings();
 	TSharedPtr<UE::Insights::TimingProfiler::STimingView> TimingViewPtr = TimingView.Pin();
 
 	TSharedRef<FTimingGraphSeries> GameFramesSeries = MakeShared<FTimingGraphSeries>(FTimingGraphSeries::ESeriesType::Frame);
@@ -326,7 +326,7 @@ void FTimingGraphTrack::AddDefaultFrameSeries()
 		GameFramesSeries->SetVisibility(Settings.GetTimingViewMainGraphShowGameFrames());
 		GameFrameSeriesVisibilityHandle = GameFramesSeries->VisibilityChangedDelegate.AddLambda([](bool bOnOff)
 			{
-				FInsightsSettings& Settings = FInsightsManager::Get()->GetSettings();
+				FInsightsSettings& Settings = UE::Insights::FInsightsManager::Get()->GetSettings();
 				Settings.SetAndSaveTimingViewMainGraphShowGameFrames(bOnOff);
 			});
 	}
@@ -345,7 +345,7 @@ void FTimingGraphTrack::AddDefaultFrameSeries()
 		RenderingFramesSeries->SetVisibility(Settings.GetTimingViewMainGraphShowRenderingFrames());
 		RenderingFrameSeriesVisibilityHandle = RenderingFramesSeries->VisibilityChangedDelegate.AddLambda([](bool bOnOff)
 			{
-				FInsightsSettings& Settings = FInsightsManager::Get()->GetSettings();
+				FInsightsSettings& Settings = UE::Insights::FInsightsManager::Get()->GetSettings();
 				Settings.SetAndSaveTimingViewMainGraphShowRenderingFrames(bOnOff);
 			});
 	}
@@ -369,7 +369,7 @@ TSharedPtr<FTimingGraphSeries> FTimingGraphTrack::GetFrameSeries(ETraceFrameType
 void FTimingGraphTrack::UpdateFrameSeries(FTimingGraphSeries& Series, const FTimingTrackViewport& Viewport)
 {
 	FGraphTrackBuilder Builder(*this, Series, Viewport);
-	TSharedPtr<const TraceServices::IAnalysisSession> Session = FInsightsManager::Get()->GetSession();
+	TSharedPtr<const TraceServices::IAnalysisSession> Session = UE::Insights::FInsightsManager::Get()->GetSession();
 	if (Session.IsValid())
 	{
 		TraceServices::FAnalysisSessionReadScope SessionReadScope(*Session.Get());
@@ -449,7 +449,7 @@ void FTimingGraphTrack::RemoveTimerSeries(uint32 TimerId)
 void FTimingGraphTrack::UpdateTimerSeries(FTimingGraphSeries& Series, const FTimingTrackViewport& Viewport)
 {
 	FGraphTrackBuilder Builder(*this, Series, Viewport);
-	TSharedPtr<const TraceServices::IAnalysisSession> Session = FInsightsManager::Get()->GetSession();
+	TSharedPtr<const TraceServices::IAnalysisSession> Session = UE::Insights::FInsightsManager::Get()->GetSession();
 	if (Session.IsValid())
 	{
 		TraceServices::FAnalysisSessionReadScope SessionReadScope(*Session.Get());
@@ -604,7 +604,7 @@ void FTimingGraphTrack::UpdateFrameStatsTimerSeries(FTimingGraphSeries& Series,
 	using namespace UE::Insights::TimingProfiler;
 
 	FGraphTrackBuilder Builder(*this, Series, Viewport);
-	TSharedPtr<const TraceServices::IAnalysisSession> Session = FInsightsManager::Get()->GetSession();
+	TSharedPtr<const TraceServices::IAnalysisSession> Session = UE::Insights::FInsightsManager::Get()->GetSession();
 	if (Session.IsValid())
 	{
 		TraceServices::FAnalysisSessionReadScope SessionReadScope(*Session.Get());
@@ -681,7 +681,7 @@ TSharedPtr<FTimingGraphSeries> FTimingGraphTrack::AddStatsCounterSeries(uint32 C
 	bool bIsMemory = false;
 	bool bIsFloatingPoint = false;
 
-	TSharedPtr<const TraceServices::IAnalysisSession> Session = FInsightsManager::Get()->GetSession();
+	TSharedPtr<const TraceServices::IAnalysisSession> Session = UE::Insights::FInsightsManager::Get()->GetSession();
 	if (Session.IsValid())
 	{
 		TraceServices::FAnalysisSessionReadScope SessionReadScope(*Session.Get());
@@ -735,7 +735,7 @@ void FTimingGraphTrack::UpdateStatsCounterSeries(FTimingGraphSeries& Series, con
 {
 	FGraphTrackBuilder Builder(*this, Series, Viewport);
 
-	TSharedPtr<const TraceServices::IAnalysisSession> Session = FInsightsManager::Get()->GetSession();
+	TSharedPtr<const TraceServices::IAnalysisSession> Session = UE::Insights::FInsightsManager::Get()->GetSession();
 	if (Session.IsValid())
 	{
 		TraceServices::FAnalysisSessionReadScope SessionReadScope(*Session.Get());
@@ -833,7 +833,7 @@ void FTimingGraphTrack::ContextMenu_ToggleOption_Execute(EGraphOptions Option)
 		return;
 	}
 
-	FInsightsSettings& Settings = FInsightsManager::Get()->GetSettings();
+	FInsightsSettings& Settings = UE::Insights::FInsightsManager::Get()->GetSettings();
 	if (EnumHasAnyFlags(Option, EGraphOptions::ShowPoints))
 	{
 		Settings.SetAndSaveTimingViewMainGraphShowPoints(EnumHasAnyFlags(EnabledOptions, EGraphOptions::ShowPoints));
@@ -867,7 +867,7 @@ void FTimingGraphTrack::LoadDefaultSettings()
 	TSharedPtr<UE::Insights::TimingProfiler::STimingView> TimingViewPtr = TimingView.Pin();
 	if (TimingViewPtr.IsValid() && TimingViewPtr->GetName() == FInsightsManagerTabs::TimingProfilerTabId)
 	{
-		const FInsightsSettings& Settings = FInsightsManager::Get()->GetSettings();
+		const FInsightsSettings& Settings = UE::Insights::FInsightsManager::Get()->GetSettings();
 		if (Settings.GetTimingViewMainGraphShowPoints())
 		{
 			EnabledOptions |= EGraphOptions::ShowPoints;
diff --git a/Engine/Source/Developer/TraceInsights/Private/Insights/Widgets/SLogView.cpp b/Engine/Source/Developer/TraceInsights/Private/Insights/Widgets/SLogView.cpp
--- a/Engine/Source/Developer/TraceInsights/Private/Insights/Widgets/SLogView.cpp
+++ b/Engine/Source/Developer/TraceInsights/Private/Insights/Widgets/SLogView.cpp
@@ -41,7 +41,10 @@
 
 ////////////////////////////////////////////////////////////////////////////////////////////////////
 
-#define LOCTEXT_NAMESPACE "SLogView"
+#define LOCTEXT_NAMESPACE "UE::Insights::SLogView"
+
+namespace UE::Insights
+{
 
 ////////////////////////////////////////////////////////////////////////////////////////////////////
 // FLogViewCommands
@@ -262,7 +265,7 @@ public:
 				FLogMessageRecord& CacheEntry = ParentWidgetPin->GetCache().Get(LogMessagePin->GetIndex());
 				const double Time = CacheEntry.GetTime();
 
-				TSharedPtr<UE::Insights::TimingProfiler::STimingView> TimingView = ParentWidgetPin->GetTimingView();
+				TSharedPtr<TimingProfiler::STimingView> TimingView = ParentWidgetPin->GetTimingView();
 				if (TimingView && TimingView->IsTimeSelectedInclusive(Time))
 				{
 					return FSlateColor(FLinearColor(0.25f, 0.5f, 1.0f, 0.25f));
@@ -300,7 +303,7 @@ public:
 				}
 			}
 
-			TSharedPtr<UE::Insights::TimingProfiler::STimingView> TimingView = ParentWidgetPin->GetTimingView();
+			TSharedPtr<TimingProfiler::STimingView> TimingView = ParentWidgetPin->GetTimingView();
 			if (TimingView && TimingView->IsTimeSelectedInclusive(Time))
 			{
 				if (IsSelected)
@@ -379,7 +382,7 @@ public:
 		if (ParentWidgetPin.IsValid() && LogMessagePin.IsValid())
 		{
 			FLogMessageRecord& CacheEntry = ParentWidgetPin->GetCache().Get(LogMessagePin->GetIndex());
-			return FSlateColor(UE::Insights::TimingProfiler::FTimeMarkerTrackBuilder::GetColorByVerbosity(CacheEntry.GetVerbosity()));
+			return FSlateColor(TimingProfiler::FTimeMarkerTrackBuilder::GetColorByVerbosity(CacheEntry.GetVerbosity()));
 		}
 		else
 		{
@@ -412,7 +415,7 @@ public:
 		if (ParentWidgetPin.IsValid() && LogMessagePin.IsValid())
 		{
 			FLogMessageRecord& CacheEntry = ParentWidgetPin->GetCache().Get(LogMessagePin->GetIndex());
-			return FSlateColor(UE::Insights::TimingProfiler::FTimeMarkerTrackBuilder::GetColorByCategory(CacheEntry.GetCategory()));
+			return FSlateColor(TimingProfiler::FTimeMarkerTrackBuilder::GetColorByCategory(CacheEntry.GetCategory()));
 		}
 		else
 		{
@@ -518,13 +521,6 @@ SLogView::SLogView()
 
 SLogView::~SLogView()
 {
-	// Remove ourselves from the profiler manager.
-	//using namespace UE::Insights::TimingProfiler;
-	//if (FTimingProfilerManager::Get().IsValid())
-	//{
-	//	//TODO: FTimimgProfilerManager::Get()->OnRequestLogViewUpdate().RemoveAll(this);
-	//}
-
 	Reset();
 }
 
@@ -2070,4 +2066,6 @@ TSharedPtr<UE::Insights::TimingProfiler::STimingView> SLogView::GetTimingView()
 
 ////////////////////////////////////////////////////////////////////////////////////////////////////
 
+} // namespace UE::Insights
+
 #undef LOCTEXT_NAMESPACE
diff --git a/Engine/Source/Developer/TraceInsights/Private/Insights/Widgets/SLogView.h b/Engine/Source/Developer/TraceInsights/Private/Insights/Widgets/SLogView.h
--- a/Engine/Source/Developer/TraceInsights/Private/Insights/Widgets/SLogView.h
+++ b/Engine/Source/Developer/TraceInsights/Private/Insights/Widgets/SLogView.h
@@ -23,6 +23,9 @@ class FMenuBuilder;
 
 namespace UE::Insights::TimingProfiler { class STimingView; }
 
+namespace UE::Insights
+{
+
 ////////////////////////////////////////////////////////////////////////////////////////////////////
 
 namespace LogViewColumns
@@ -83,7 +86,7 @@ public:
 
 	bool IsFilteringAsyncTaskCancelRequested() const { return bIsFilteringAsyncTaskCancelRequested; }
 
-	TSharedPtr<UE::Insights::TimingProfiler::STimingView> GetTimingView() const;
+	TSharedPtr<TimingProfiler::STimingView> GetTimingView() const;
 
 protected:
 	/** Generate a new list view row. */
@@ -177,7 +180,7 @@ protected:
 	mutable volatile bool bIsFilteringAsyncTaskCancelRequested; // true if we want the async task to finish asap
 
 	/** Stopwatch used to measure how long it takes to filter the message list. */
-	mutable UE::Insights::FStopwatch FilteringStopwatch;
+	mutable FStopwatch FilteringStopwatch;
 
 	/**
 	 * Total number of log categories (including duplicates; last value read from LogProvider).
@@ -201,7 +204,7 @@ protected:
 	bool bIsDirty;
 
 	/** Stopwatch used to measure the time since the list of messages has become dirty. */
-	mutable UE::Insights::FStopwatch DirtyStopwatch;
+	mutable FStopwatch DirtyStopwatch;
 
 	/** Stats */
 	FText StatsText;
@@ -214,3 +217,5 @@ protected:
 };
 
 ////////////////////////////////////////////////////////////////////////////////////////////////////
+
+} // namespace UE::Insights
diff --git a/Engine/Source/Developer/TraceInsights/Private/Insights/Widgets/SMajorTabWindow.cpp b/Engine/Source/Developer/TraceInsights/Private/Insights/Widgets/SMajorTabWindow.cpp
--- a/Engine/Source/Developer/TraceInsights/Private/Insights/Widgets/SMajorTabWindow.cpp
+++ b/Engine/Source/Developer/TraceInsights/Private/Insights/Widgets/SMajorTabWindow.cpp
@@ -22,13 +22,15 @@
 	#include "Interfaces/IAnalyticsProvider.h"
 #endif // WITH_EDITOR
 
+// TraceInsightsCore
+#include "InsightsCore/Version.h"
+
 // TraceInsights
 #include "Insights/Common/InsightsMenuBuilder.h"
 #include "Insights/InsightsManager.h"
 #include "Insights/InsightsStyle.h"
 #include "Insights/Log.h"
 #include "Insights/TraceInsightsModule.h"
-#include "Insights/Version.h"
 
 ////////////////////////////////////////////////////////////////////////////////////////////////////
 
@@ -153,6 +155,8 @@ void SMajorTabWindow::RegisterTabSpawners()
 BEGIN_SLATE_FUNCTION_BUILD_OPTIMIZATION
 void SMajorTabWindow::Construct(const FArguments& InArgs, const TSharedRef<SDockTab>& ConstructUnderMajorTab, const TSharedPtr<SWindow>& ConstructUnderWindow)
 {
+	using namespace UE::Insights;
+
 	//////////////////////////////////////////////////
 	// Create & initialize tab manager.
 
@@ -343,7 +347,7 @@ void SMajorTabWindow::FillMenu(FMenuBuilder& MenuBuilder, const TSharedPtr<FTabM
 		return;
 	}
 
-	FInsightsManager::Get()->GetInsightsMenuBuilder()->PopulateMenu(MenuBuilder);
+	UE::Insights::FInsightsManager::Get()->GetInsightsMenuBuilder()->PopulateMenu(MenuBuilder);
 
 	TabManager->PopulateLocalTabSpawnerMenu(MenuBuilder);
 }
@@ -352,7 +356,7 @@ void SMajorTabWindow::FillMenu(FMenuBuilder& MenuBuilder, const TSharedPtr<FTabM
 
 bool SMajorTabWindow::IsValidSession() const
 {
-	return FInsightsManager::Get()->GetSession().IsValid();
+	return UE::Insights::FInsightsManager::Get()->GetSession().IsValid();
 }
 
 ////////////////////////////////////////////////////////////////////////////////////////////////////
@@ -408,7 +412,7 @@ FReply SMajorTabWindow::OnKeyDown(const FGeometry& MyGeometry, const FKeyEvent&
 
 FReply SMajorTabWindow::OnDragOver(const FGeometry& MyGeometry, const FDragDropEvent& DragDropEvent)
 {
-	if (FInsightsManager::Get()->OnDragOver(DragDropEvent))
+	if (UE::Insights::FInsightsManager::Get()->OnDragOver(DragDropEvent))
 	{
 		return FReply::Handled();
 	}
@@ -420,7 +424,7 @@ FReply SMajorTabWindow::OnDragOver(const FGeometry& MyGeometry, const FDragDropE
 
 FReply SMajorTabWindow::OnDrop(const FGeometry& MyGeometry, const FDragDropEvent& DragDropEvent)
 {
-	if (FInsightsManager::Get()->OnDrop(DragDropEvent))
+	if (UE::Insights::FInsightsManager::Get()->OnDrop(DragDropEvent))
 	{
 		return FReply::Handled();
 	}
diff --git a/Engine/Source/Developer/TraceInsights/Private/Insights/Widgets/SSessionInfoWindow.cpp b/Engine/Source/Developer/TraceInsights/Private/Insights/Widgets/SSessionInfoWindow.cpp
--- a/Engine/Source/Developer/TraceInsights/Private/Insights/Widgets/SSessionInfoWindow.cpp
+++ b/Engine/Source/Developer/TraceInsights/Private/Insights/Widgets/SSessionInfoWindow.cpp
@@ -24,16 +24,19 @@
 
 // TraceInsightsCore
 #include "InsightsCore/Common/TimeUtils.h"
+#include "InsightsCore/Version.h"
 
 // TraceInsights
 #include "Insights/Common/InsightsMenuBuilder.h"
 #include "Insights/InsightsManager.h"
 #include "Insights/InsightsStyle.h"
-#include "Insights/Version.h"
 
 ////////////////////////////////////////////////////////////////////////////////////////////////////
 
-#define LOCTEXT_NAMESPACE "SSessionInfoWindow"
+#define LOCTEXT_NAMESPACE "UE::Insights::SSessionInfoWindow"
+
+namespace UE::Insights
+{
 
 ////////////////////////////////////////////////////////////////////////////////////////////////////
 
@@ -616,4 +619,6 @@ FText SSessionInfoWindow::GetModulesText() const
 
 ////////////////////////////////////////////////////////////////////////////////////////////////////
 
+} // namespace UE::Insights
+
 #undef LOCTEXT_NAMESPACE
diff --git a/Engine/Source/Developer/TraceInsights/Private/Insights/Widgets/STimingView.cpp b/Engine/Source/Developer/TraceInsights/Private/Insights/Widgets/STimingView.cpp
--- a/Engine/Source/Developer/TraceInsights/Private/Insights/Widgets/STimingView.cpp
+++ b/Engine/Source/Developer/TraceInsights/Private/Insights/Widgets/STimingView.cpp
@@ -76,7 +76,6 @@
 #include "Insights/TimingProfilerCommon.h"
 #include "Insights/ViewModels/BaseTimingTrack.h"
 #include "Insights/ViewModels/DrawHelpers.h"
-#include "Insights/ViewModels/EventNameFilterValueConverter.h"
 #include "Insights/ViewModels/GraphSeries.h"
 #include "Insights/ViewModels/GraphTrack.h"
 #include "Insights/ViewModels/QuickFind.h"
diff --git a/Engine/Source/Developer/TraceInsights/Public/Insights/IUnrealInsightsModule.h b/Engine/Source/Developer/TraceInsights/Public/Insights/IUnrealInsightsModule.h
--- a/Engine/Source/Developer/TraceInsights/Public/Insights/IUnrealInsightsModule.h
+++ b/Engine/Source/Developer/TraceInsights/Public/Insights/IUnrealInsightsModule.h
@@ -26,14 +26,16 @@ namespace Insights { class IInsightsManager; }
 struct TRACEINSIGHTS_API FInsightsManagerTabs
 {
 	static const FName StartPageTabId; // DEPRECATED
-	static const FName TraceStoreTabId;
-	static const FName ConnectionTabId;
-	static const FName LauncherTabId;
+	static const FName TraceStoreTabId; // DEPRECATED: see FInsightsFrontendTabs::TraceStoreTabId
+	static const FName ConnectionTabId; // DEPRECATED: see FInsightsFrontendTabs::ConnectionTabId
+	static const FName LauncherTabId; // DEPRECATED
+
 	static const FName SessionInfoTabId;
 	static const FName TimingProfilerTabId;
 	static const FName LoadingProfilerTabId;
-	static const FName NetworkingProfilerTabId;
 	static const FName MemoryProfilerTabId;
+	static const FName NetworkingProfilerTabId;
+
 	static const FName AutomationWindowTabId;
 	static const FName MessageLogTabId;
 };
@@ -134,17 +136,6 @@ protected:
 
 ////////////////////////////////////////////////////////////////////////////////////////////////////
 
-/** Contains parameters that are passed to the CreateSessionBrowser function to control specific behaviors. */
-struct TRACEINSIGHTS_API FCreateSessionBrowserParams
-{
-	bool bAllowDebugTools = false;
-	bool bInitializeTesting = false;
-	bool bStartProcessWithStompMalloc = false;
-	bool bDisableFramerateThrottle = false;
-};
-
-////////////////////////////////////////////////////////////////////////////////////////////////////
-
 /** Called back to register common layout extensions */
 DECLARE_MULTICAST_DELEGATE_OneParam(FOnRegisterMajorTabExtensions, FInsightsMajorTabExtender& /*MajorTabExtender*/);
 
@@ -270,11 +261,6 @@ public:
 	/** Sets the ini path for saving persistent layout data. */
 	virtual void SetUnrealInsightsLayoutIni(const FString& InIniPath) = 0;
 
-	/**
-	 * Called when the application starts in "Browser" mode.
-	 */
-	virtual void CreateSessionBrowser(const FCreateSessionBrowserParams& Params) = 0;
-
 	/**
 	 * Called when the application starts in "Viewer" mode.
 	 */
@@ -290,11 +276,6 @@ public:
 	*/
 	virtual void ScheduleCommand(const FString& InCmd) = 0;
 
-	/**
-	* Called to run automation test in Insights.
-	*/
-	virtual void RunAutomationTest(const FString& InCmd) = 0;
-
 	/**
 	* Called to initialize testing in stand alone Insights.
 	 * @param InInitAutomationModules If true Insights will initialize the modules required for running automation tests.
diff --git a/Engine/Source/Developer/TraceInsights/Public/Insights/Version.h b/Engine/Source/Developer/TraceInsights/Public/Insights/Version.h
--- a/Engine/Source/Developer/TraceInsights/Public/Insights/Version.h
+++ b/Engine/Source/Developer/TraceInsights/Public/Insights/Version.h
@@ -2,19 +2,8 @@
 
 #pragma once
 
-// See https://semver.org/
+#include "Misc/CoreMiscDefines.h"
+UE_DEPRECATED_HEADER(5.5, "Use InsightsCore/Version.h from the TraceInsightsCore module instead.")
 
-#define UNREAL_INSIGHTS_VERSION_MAJOR 1
-#define UNREAL_INSIGHTS_VERSION_MINOR 5
-#define UNREAL_INSIGHTS_VERSION_PATCH 0
-
-#define UNREAL_INSIGHTS_VERSION_STRING "1.05"
-
-/** Extra identifier string (like: "alpha", "beta.1", "beta.2", "test", "featureX", etc.). */
-//#define UNREAL_INSIGHTS_VERSION_ID_STRING "beta.1"
-
-#ifdef UNREAL_INSIGHTS_VERSION_ID_STRING
-#define UNREAL_INSIGHTS_VERSION_STRING_EX "v" UNREAL_INSIGHTS_VERSION_STRING "-" UNREAL_INSIGHTS_VERSION_ID_STRING
-#else
-#define UNREAL_INSIGHTS_VERSION_STRING_EX "v" UNREAL_INSIGHTS_VERSION_STRING
-#endif
+// TraceInsightsCore
+#include "InsightsCore/Version.h"
diff --git a/Engine/Source/Developer/TraceInsights/TraceInsights.Build.cs b/Engine/Source/Developer/TraceInsights/TraceInsights.Build.cs
--- a/Engine/Source/Developer/TraceInsights/TraceInsights.Build.cs
+++ b/Engine/Source/Developer/TraceInsights/TraceInsights.Build.cs
@@ -10,7 +10,7 @@ public class TraceInsights : ModuleRules
 
 		PublicDependencyModuleNames.AddRange(
 			new string[] {
-				"AppFramework", // for SColorPicker
+				"AppFramework", // for SColorPicker, RestoreStarshipSuite()
 				"ApplicationCore",
 				"AutomationDriver",
 				"Cbor",
@@ -32,7 +32,7 @@ public class TraceInsights : ModuleRules
 			}
 		);
 
-		//Modules required for running automation in stand alone Insights
+		// Modules required for running automation in stand alone Insights
 		if (Target.Configuration != UnrealTargetConfiguration.Shipping && !Target.bBuildEditor)
 		{
 			PrivateDependencyModuleNames.AddRange(
diff --git a/Engine/Source/Developer/TraceInsightsCore/Private/InsightsCore/Common/InsightsCoreStyle.cpp b/Engine/Source/Developer/TraceInsightsCore/Private/InsightsCore/Common/InsightsCoreStyle.cpp
--- a/Engine/Source/Developer/TraceInsightsCore/Private/InsightsCore/Common/InsightsCoreStyle.cpp
+++ b/Engine/Source/Developer/TraceInsightsCore/Private/InsightsCore/Common/InsightsCoreStyle.cpp
@@ -69,7 +69,7 @@ void FInsightsCoreStyle::Shutdown()
 
 FName FInsightsCoreStyle::GetStyleSetName()
 {
-	static FName StyleSetName(TEXT("TraceInsightsCore/InsightsStyle"));
+	static FName StyleSetName(TEXT("InsightsCoreStyle"));
 	return StyleSetName;
 }
 
@@ -123,21 +123,69 @@ void FInsightsCoreStyle::FStyle::Initialize()
 	const FVector2D Icon20x20(20.0f, 20.0f); // for ToolBar icons
 
 	//////////////////////////////////////////////////
+	// Color brushes
+
+	Set("DarkGreenBrush", new FSlateColorBrush(FLinearColor(0.0f, 0.25f, 0.0f, 1.0f)));
+
+	//////////////////////////////////////////////////
+	// Border brushes
+
+	Set("SingleBorder", new FSlateBorderBrush(NAME_None, FMargin(1.0f)));
+	Set("DoubleBorder", new FSlateBorderBrush(NAME_None, FMargin(2.0f)));
+
+	Set("EventBorder", new FSlateBorderBrush(NAME_None, FMargin(1.0f)));
+	Set("HoveredEventBorder", new FSlateBorderBrush(NAME_None, FMargin(2.0f)));
+	Set("SelectedEventBorder", new FSlateBorderBrush(NAME_None, FMargin(2.0f)));
+
+	//////////////////////////////////////////////////
+	// Box brushes
+
+	Set("RoundedBackground", new FSlateRoundedBoxBrush(FLinearColor(1.0f, 1.0f, 1.0f, 1.0f), Icon16x16));
 
 	Set("Border.TB", new CORE_BOX_BRUSH("Icons/Profiler/Profiler_Border_TB_16x", FMargin(4.0f / 16.0f)));
 	Set("Border.L", new CORE_BOX_BRUSH("Icons/Profiler/Profiler_Border_L_16x", FMargin(4.0f / 16.0f)));
 	Set("Border.R", new CORE_BOX_BRUSH("Icons/Profiler/Profiler_Border_R_16x", FMargin(4.0f / 16.0f)));
 
+	//////////////////////////////////////////////////
+	// Timing View resources
+
+	//////////////////////////////////////////////////
+	// Graph Track resources
+
+	Set("Graph.Point", new EDITOR_IMAGE_BRUSH("Old/Graph/ExecutionBubble", Icon16x16));
+
+	//////////////////////////////////////////////////
+
+	Set("Icons.Debug", new EDITOR_IMAGE_BRUSH_SVG("Starship/Common/bug", Icon16x16));
+	Set("Icons.Debug.ToolBar", new EDITOR_IMAGE_BRUSH_SVG("Starship/Common/bug", Icon20x20));
+
+	Set("Icons.FolderExplore", new EDITOR_IMAGE_BRUSH_SVG("Starship/Common/ContentBrowser", Icon16x16));
+	//Set("Icons.FolderOpen", new CORE_IMAGE_BRUSH_SVG("Starship/Common/folder-open", Icon16x16));		//-> use FAppStyle "Icons.FolderOpen"
+	//Set("Icons.FolderClosed", new CORE_IMAGE_BRUSH_SVG("Starship/Common/folder-closed", Icon16x16));	//-> use FAppStyle "Icons.FolderClosed"
+
+	Set("Icons.ResetToDefault", new CORE_IMAGE_BRUSH("Icons/Profiler/profiler_ResetToDefault_32x", Icon16x16));
+	Set("Icons.DiffersFromDefault", new EDITOR_IMAGE_BRUSH_SVG("Starship/Common/ResetToDefault", Icon16x16));
+
+	Set("Icons.TestAutomation", new EDITOR_IMAGE_BRUSH_SVG("Starship/Common/TestAutomation", Icon16x16));
+	Set("Icons.Test", new CORE_IMAGE_BRUSH_SVG("Starship/Common/Test", Icon16x16));
+
+	Set("Icons.Rename", new CORE_IMAGE_BRUSH_SVG("Starship/Common/Rename", Icon16x16));
+	//Set("Icons.Delete", new CORE_IMAGE_BRUSH_SVG("Starship/Common/Delete", Icon16x16));	//-> use FAppStyle "Icons.Delete"
+
+	Set("Icons.Find", new EDITOR_IMAGE_BRUSH_SVG("Starship/Common/TraceDataFiltering", Icon16x16));
+
 	//////////////////////////////////////////////////
 
 	Set("Icons.TableTreeView", new EDITOR_IMAGE_BRUSH_SVG("Starship/Common/Spreadsheet", Icon16x16));
 	Set("Icons.TableTreeView.ToolBar", new EDITOR_IMAGE_BRUSH_SVG("Starship/Common/Spreadsheet", Icon20x20));
 
+	Set("Icons.Filter.ToolBar", new CORE_IMAGE_BRUSH_SVG("Starship/Common/filter", Icon20x20));
 	//Set("Icons.Filter", new CORE_IMAGE_BRUSH_SVG("Starship/Common/filter", Icon16x16));	//-> use FAppStyle "Icons.Filter"
 	Set("Icons.FilterAddGroup", new EDITOR_IMAGE_BRUSH_SVG("Starship/Common/WorldOutliner", Icon16x16));
 	Set("Icons.ClassicFilter", new IMAGE_BRUSH_SVG("Filter", Icon16x16));
 	Set("Icons.ClassicFilterConfig", new IMAGE_BRUSH_SVG("FilterConfig", Icon16x16));
 
+
 	Set("Icons.SortBy", new CORE_IMAGE_BRUSH("Icons/Profiler/profiler_SortBy_32x", Icon16x16));
 	//Set("Icons.SortUp", new CORE_IMAGE_BRUSH_SVG("Starship/Common/SortUp", Icon16x16));		//-> use FAppStyle "Icons.SortUp"
 	//Set("Icons.SortDown", new CORE_IMAGE_BRUSH_SVG("Starship/Common/SortDown", Icon16x16));	//-> use FAppStyle "Icons.SortDown"
@@ -150,14 +198,14 @@ void FInsightsCoreStyle::FStyle::Initialize()
 	Set("Icons.ExpandSelection", new CORE_IMAGE_BRUSH("Icons/Profiler/profiler_ExpandSelection_32x", Icon16x16));
 	Set("Icons.CollapseSelection", new CORE_IMAGE_BRUSH("Icons/Profiler/profiler_CollapseSelection_32x", Icon16x16));
 
-	Set("Icons.AddGraphSeries", new CORE_IMAGE_BRUSH_SVG("Starship/Common/plus", Icon16x16));
-	Set("Icons.RemoveGraphSeries", new CORE_IMAGE_BRUSH_SVG("Starship/Common/close", Icon16x16));
-
-	Set("Icons.Delete", new CORE_IMAGE_BRUSH_SVG("Starship/Common/Delete", Icon16x16));
+	//////////////////////////////////////////////////
 
 	Set("TreeTable.RowBackground", new EDITOR_IMAGE_BRUSH("Old/White", Icon16x16, FLinearColor(1.0f, 1.0f, 1.0f, 0.25f)));
 	Set("TreeViewBanner.WarningIcon", new CORE_IMAGE_BRUSH_SVG("Starship/Common/alert-circle", Icon20x20, FStyleColors::Warning));
 
+	//////////////////////////////////////////////////
+	// Icons for tree/table items
+
 	Set("Icons.Hint.TreeItem", new IMAGE_BRUSH_SVG("InfoTag_12", Icon12x12));
 	Set("Icons.Group.TreeItem", new CORE_IMAGE_BRUSH_SVG("Starship/Common/folder-closed", Icon12x12));
 	Set("Icons.Leaf.TreeItem", new CORE_IMAGE_BRUSH_SVG("Starship/Common/bullet-point", Icon12x12));
@@ -166,6 +214,8 @@ void FInsightsCoreStyle::FStyle::Initialize()
 	Set("Icons.Plugin.TreeItem", new EDITOR_IMAGE_BRUSH_SVG("Starship/Common/Plugins", Icon12x12));
 	Set("Icons.Dependencies.TreeItem", new EDITOR_IMAGE_BRUSH_SVG("Starship/Common/Blueprint", Icon12x12));
 
+	//////////////////////////////////////////////////
+
 	Set("TreeTable.TooltipBold", FTextBlockStyle(NormalText)
 		.SetFont(DEFAULT_FONT("Bold", 8))
 		.SetColorAndOpacity(FLinearColor(0.5f, 0.5f, 0.5f, 1.0f))
@@ -191,7 +241,6 @@ void FInsightsCoreStyle::FStyle::Initialize()
 	//////////////////////////////////////////////////
 
 	// NormalEditableTextBox && SearchBox
-#if 0
 	{
 		const FEditableTextBoxStyle& NormalEditableTextBoxStyle = FAppStyle::GetWidgetStyle<FEditableTextBoxStyle>("NormalEditableTextBox");
 		const FTextBlockStyle TextBlockStyle = FTextBlockStyle()
@@ -209,10 +258,8 @@ void FInsightsCoreStyle::FStyle::Initialize()
 			.SetTextBoxStyle(EditableTextBoxStyle);
 		Set("SearchBox", SearchBoxStyle);
 	}
-#endif
 
 	// PrimaryToolbar
-#if 0
 	{
 		FToolBarStyle PrimaryToolbarStyle = FAppStyle::Get().GetWidgetStyle<FToolBarStyle>("SlimToolBar");
 
@@ -221,10 +268,8 @@ void FInsightsCoreStyle::FStyle::Initialize()
 		Set("PrimaryToolbar.MinUniformToolbarSize", 40.0f);
 		Set("PrimaryToolbar.MaxUniformToolbarSize", 40.0f);
 	}
-#endif
 
 	// SecondaryToolbar
-#if 0
 	{
 		FToolBarStyle SecondaryToolbarStyle = FAppStyle::Get().GetWidgetStyle<FToolBarStyle>("SlimToolBar");
 
@@ -249,10 +294,8 @@ void FInsightsCoreStyle::FStyle::Initialize()
 		Set("SecondaryToolbar.MinUniformToolbarSize", 32.0f);
 		Set("SecondaryToolbar.MaxUniformToolbarSize", 32.0f);
 	}
-#endif
 
 	// SecondaryToolbar2 (used by AutoScroll and NetPacketContentView toolbars)
-#if 0
 	{
 		FToolBarStyle SecondaryToolbarStyle = FAppStyle::Get().GetWidgetStyle<FToolBarStyle>("SlimToolBar");
 
@@ -277,20 +320,16 @@ void FInsightsCoreStyle::FStyle::Initialize()
 		Set("SecondaryToolbar2.MinUniformToolbarSize", 32.0f);
 		Set("SecondaryToolbar2.MaxUniformToolbarSize", 32.0f);
 	}
-#endif
 
 	// ToggleButton
-#if 0
 	{
 		Set("ToggleButton", FButtonStyle(Button)
 			.SetNormal(FSlateNoResource())
 			.SetHovered(EDITOR_BOX_BRUSH("Common/RoundedSelection_16x", 4.0f / 16.0f, SelectionColor))
 			.SetPressed(EDITOR_BOX_BRUSH("Common/RoundedSelection_16x", 4.0f / 16.0f, SelectionColor_Pressed)));
 	}
-#endif
 
 	// Common.GotoNativeCodeHyperlink
-#if 0
 	{
 		FTextBlockStyle InheritedFromNativeTextStyle = FTextBlockStyle(NormalText)
 			.SetFont(DEFAULT_FONT("Regular", 10));
@@ -309,7 +348,6 @@ void FInsightsCoreStyle::FStyle::Initialize()
 
 		Set("Common.GotoNativeCodeHyperlink", EditNativeHyperlinkStyle);
 	}
-#endif
 
 	//////////////////////////////////////////////////
 }
diff --git a/Engine/Source/Developer/TraceInsightsCore/TraceInsightsCore.Build.cs b/Engine/Source/Developer/TraceInsightsCore/TraceInsightsCore.Build.cs
--- a/Engine/Source/Developer/TraceInsightsCore/TraceInsightsCore.Build.cs
+++ b/Engine/Source/Developer/TraceInsightsCore/TraceInsightsCore.Build.cs
@@ -17,6 +17,7 @@ public class TraceInsightsCore : ModuleRules
 				"InputCore",
 				"Slate",
 				"SlateCore",
+				"ToolWidgets",
 			}
 		);
 	}
diff --git a/Engine/Source/Developer/TraceInsightsFrontend/Private/InsightsFrontend/Tests/TestRunner.cpp b/Engine/Source/Developer/TraceInsightsFrontend/Private/InsightsFrontend/Tests/TestRunner.cpp
new file mode 100644
--- /dev/null
+++ b/Engine/Source/Developer/TraceInsightsFrontend/Private/InsightsFrontend/Tests/TestRunner.cpp
@@ -0,0 +1,46 @@
+// Copyright Epic Games, Inc. All Rights Reserved.
+
+#include "TestRunner.h"
+
+#if !UE_BUILD_SHIPPING && !WITH_EDITOR
+
+#include "IAutomationControllerModule.h"
+#include "IAutomationControllerManager.h"
+#include "UObject/UObjectGlobals.h" // for StaticExec()
+
+////////////////////////////////////////////////////////////////////////////////////////////////////
+
+DEFINE_LOG_CATEGORY(InsightsTestRunner);
+
+#define LOCTEXT_NAMESPACE "UE::Insights::TestRunner"
+
+namespace UE::Insights
+{
+
+////////////////////////////////////////////////////////////////////////////////////////////////////
+
+void FTestRunner::Run(const FString& InCmd)
+{
+	FString ActualCmd = InCmd.Replace(TEXT("\""), TEXT(""));
+	if (!ActualCmd.StartsWith(TEXT("Automation RunTests")))
+	{
+		UE_LOG(InsightsTestRunner, Warning, TEXT("[InsightsTestRunner] Command %s does not start with Automation RunTests. Command will be ignored."), *InCmd);
+		return;
+	}
+
+	IAutomationControllerModule& AutomationControllerModule = FModuleManager::LoadModuleChecked<IAutomationControllerModule>(TEXT("AutomationController"));
+	IAutomationControllerManagerRef AutomationControllerManager = AutomationControllerModule.GetAutomationController();
+
+	AutomationControllerManager->OnTestsComplete().AddLambda([this]() { bIsRunningTests = false; });
+	bIsRunningTests = true;
+
+	::StaticExec(NULL, *ActualCmd);
+}
+
+////////////////////////////////////////////////////////////////////////////////////////////////////
+
+} // namespace UE::Insights
+
+#undef LOCTEXT_NAMESPACE
+
+#endif //UE_BUILD_SHIPPING && !WITH_EDITOR
diff --git a/Engine/Source/Developer/TraceInsightsFrontend/Private/InsightsFrontend/Tests/TestRunner.h b/Engine/Source/Developer/TraceInsightsFrontend/Private/InsightsFrontend/Tests/TestRunner.h
new file mode 100644
--- /dev/null
+++ b/Engine/Source/Developer/TraceInsightsFrontend/Private/InsightsFrontend/Tests/TestRunner.h
@@ -0,0 +1,31 @@
+// Copyright Epic Games, Inc. All Rights Reserved.
+
+#pragma once
+
+#include "CoreTypes.h"
+
+#if !UE_BUILD_SHIPPING && !WITH_EDITOR
+
+#include "Logging/LogMacros.h"
+
+#include "InsightsFrontend/ITraceInsightsFrontendModule.h"
+
+DECLARE_LOG_CATEGORY_EXTERN(InsightsTestRunner, Log, All);
+
+namespace UE::Insights
+{
+
+class FTestRunner
+{
+public:
+	virtual ~FTestRunner() {}
+
+	void Run(const FString& InCmd);
+
+private:
+	bool bIsRunningTests = false;
+};
+
+} // namespace UE::Insights
+
+#endif //UE_BUILD_SHIPPING && !WITH_EDITOR
diff --git a/Engine/Source/Programs/UnrealInsights/Private/UnrealInsightsMain.cpp b/Engine/Source/Programs/UnrealInsights/Private/UnrealInsightsMain.cpp
--- a/Engine/Source/Programs/UnrealInsights/Private/UnrealInsightsMain.cpp
+++ b/Engine/Source/Programs/UnrealInsights/Private/UnrealInsightsMain.cpp
@@ -6,8 +6,107 @@
 #include "RequiredProgramMainCPPInclude.h"
 #include "UserInterfaceCommand.h"
 
+#if PLATFORM_WINDOWS
+#include "Windows/AllowWindowsPlatformTypes.h"
+#include "Windows/HideWindowsPlatformTypes.h"
+#endif
+
+#if PLATFORM_UNIX
+#include <sys/file.h>
+#include <errno.h>
+#endif
+
 IMPLEMENT_APPLICATION(UnrealInsights, "UnrealInsights");
 
+////////////////////////////////////////////////////////////////////////////////////////////////////
+
+FString GetTraceFileFromCmdLine(const TCHAR* CommandLine)
+{
+	const uint32 MaxPath = FPlatformMisc::GetMaxPathLength();
+	TCHAR* FilenameBuffer = new TCHAR[MaxPath + 1];
+	FilenameBuffer[0] = 0;
+
+	// Try getting the trace file from the -OpenTraceFile= parameter first.
+	if (FParse::Value(CommandLine, TEXT("-OpenTraceFile="), FilenameBuffer, MaxPath, true))
+	{
+		FString TraceFile(FilenameBuffer);
+		delete[] FilenameBuffer;
+		return TraceFile;
+	}
+
+	// Support for opening a trace file by double clicking a .utrace file.
+	// In this case, the app will receive as the first parameter a utrace file path.
+
+	const TCHAR* Str = CommandLine;
+	if (FParse::Token(Str, FilenameBuffer, MaxPath, false))
+	{
+		FString TraceFile(FilenameBuffer);
+		if (TraceFile.EndsWith(TEXT(".utrace")))
+		{
+			delete[] FilenameBuffer;
+			return TraceFile;
+		}
+	}
+
+	delete[] FilenameBuffer;
+	return FString();
+}
+
+////////////////////////////////////////////////////////////////////////////////////////////////////
+
+bool CheckFrontendSingleInstance()
+{
+#if PLATFORM_WINDOWS
+	// Create a named event that other processes can detect.
+	// It allows only a single instance of Unreal Insights Frontend.
+	HANDLE SessionBrowserEvent = CreateEvent(NULL, true, false, TEXT("Local\\UnrealInsightsBrowser"));
+	if (SessionBrowserEvent == NULL || GetLastError() == ERROR_ALREADY_EXISTS)
+	{
+		// Another Session Browser process is already running.
+
+		if (SessionBrowserEvent != NULL)
+		{
+			CloseHandle(SessionBrowserEvent);
+		}
+
+		// Activate the respective window.
+		HWND Window = FindWindowW(0, L"Unreal Insights Frontend");
+		if (Window)
+		{
+			ShowWindow(Window, SW_SHOW);
+			SetForegroundWindow(Window);
+
+			FLASHWINFO FlashInfo;
+			FlashInfo.cbSize = sizeof(FLASHWINFO);
+			FlashInfo.hwnd = Window;
+			FlashInfo.dwFlags = FLASHW_ALL;
+			FlashInfo.uCount = 3;
+			FlashInfo.dwTimeout = 0;
+			FlashWindowEx(&FlashInfo);
+		}
+
+		return false;
+	}
+#endif // PLATFORM_WINDOWS
+
+#if PLATFORM_UNIX
+	int FileHandle = open("/var/run/UnrealInsightsBrowser.pid", O_CREAT | O_RDWR, 0666);
+	int Ret = flock(FileHandle, LOCK_EX | LOCK_NB);
+	if (Ret && EWOULDBLOCK == errno)
+	{
+		// Another Session Browser process is already running.
+
+		// Activate the respective window.
+		//TODO: "wmctrl -a Insights"
+
+		return false;
+	}
+#endif
+
+	return true;
+}
+
+////////////////////////////////////////////////////////////////////////////////////////////////////
 /**
  * Platform agnostic implementation of the main entry point.
  */
@@ -18,6 +117,25 @@ int32 UnrealInsightsMain(const TCHAR* CommandLine)
 	// Override the stack size for the thread pool.
 	FQueuedThreadPool::OverrideStackSize = 256 * 1024;
 
+	bool bFrontendMode;
+	FString TraceFileToOpen;
+
+	if (FCString::Strifind(CommandLine, TEXT("-OpenTraceId=")))
+	{
+		bFrontendMode = false;
+	}
+	else
+	{
+		TraceFileToOpen = GetTraceFileFromCmdLine(CommandLine);
+		bFrontendMode = TraceFileToOpen.IsEmpty();
+
+		// Only a single instance of Unreal Insights Frontend window/process is allowed.
+		if (bFrontendMode && !CheckFrontendSingleInstance())
+		{
+			return 0;
+		}
+	}
+
 	FString NewCommandLine = CommandLine;
 
 	// Add -Messaging if it was not given in the command line.
@@ -29,18 +147,16 @@ int32 UnrealInsightsMain(const TCHAR* CommandLine)
 	// Initialize core.
 	GEngineLoop.PreInit(*NewCommandLine);
 
-	// Make sure all UObject classes are registered and default properties have been initialized.
-	//ProcessNewlyLoadedUObjects();
-
 	// Tell the module manager it may now process newly-loaded UObjects when new C++ modules are loaded.
 	FModuleManager::Get().StartProcessingNewlyLoadedObjects();
 
-	FUserInterfaceCommand::Run();
+	FUserInterfaceCommand::Run(bFrontendMode, TraceFileToOpen);
 
 	// Shut down.
-	//FCoreDelegates::OnExit.Broadcast();
 	FEngineLoop::AppPreExit();
 	FModuleManager::Get().UnloadModulesAtShutdown();
 
 	return 0;
 }
+
+////////////////////////////////////////////////////////////////////////////////////////////////////
diff --git a/Engine/Source/Programs/UnrealInsights/Private/UserInterfaceCommand.cpp b/Engine/Source/Programs/UnrealInsights/Private/UserInterfaceCommand.cpp
--- a/Engine/Source/Programs/UnrealInsights/Private/UserInterfaceCommand.cpp
+++ b/Engine/Source/Programs/UnrealInsights/Private/UserInterfaceCommand.cpp
@@ -6,7 +6,6 @@
 //#include "Brushes/SlateImageBrush.h"
 #include "Containers/Ticker.h"
 #include "CoreGlobals.h"
-#include "Styling/AppStyle.h"
 #include "Framework/Application/SlateApplication.h"
 #include "Framework/Docking/LayoutService.h"
 #include "Framework/Docking/TabManager.h"
@@ -20,20 +19,17 @@
 #include "Misc/Parse.h"
 #include "Modules/ModuleManager.h"
 #include "StandaloneRenderer.h"
+#include "Styling/AppStyle.h"
 #include "Widgets/Docking/SDockTab.h"
 
-// Insights
+// TraceInsightsCore
+#include "InsightsCore/Version.h"
+
+// TraceInsights
 #include "Insights/IUnrealInsightsModule.h"
-#include "Insights/Version.h"
 
-#if PLATFORM_WINDOWS
-#include "Windows/AllowWindowsPlatformTypes.h"
-#include "Windows/HideWindowsPlatformTypes.h"
-#endif
-#if PLATFORM_UNIX
-#include <sys/file.h>
-#include <errno.h>
-#endif
+// TraceInsightsFrontend
+#include "InsightsFrontend/ITraceInsightsFrontendModule.h"
 
 ////////////////////////////////////////////////////////////////////////////////////////////////////
 
@@ -74,97 +70,23 @@ namespace UserInterfaceCommand
 
 ////////////////////////////////////////////////////////////////////////////////////////////////////
 
-bool CheckSessionBrowserSingleInstance()
+void FUserInterfaceCommand::Run(bool bFrontendMode, const FString& TraceFileToOpen)
 {
-#if PLATFORM_WINDOWS
-	// Create a named event that other processes can detect.
-	// It allows only a single instance of Unreal Insights (Browser Mode).
-	HANDLE SessionBrowserEvent = CreateEvent(NULL, true, false, TEXT("Local\\UnrealInsightsBrowser"));
-	if (SessionBrowserEvent == NULL || GetLastError() == ERROR_ALREADY_EXISTS)
-	{
-		// Another Session Browser process is already running.
-
-		if (SessionBrowserEvent != NULL)
-		{
-			CloseHandle(SessionBrowserEvent);
-		}
-
-		// Activate the respective window.
-		HWND Window = FindWindowW(0, L"Unreal Insights Session Browser");
-		if (Window)
-		{
-			ShowWindow(Window, SW_SHOW);
-			SetForegroundWindow(Window);
-
-			FLASHWINFO FlashInfo;
-			FlashInfo.cbSize = sizeof(FLASHWINFO);
-			FlashInfo.hwnd = Window;
-			FlashInfo.dwFlags = FLASHW_ALL;
-			FlashInfo.uCount = 3;
-			FlashInfo.dwTimeout = 0;
-			FlashWindowEx(&FlashInfo);
-		}
-
-		return false;
-	}
-#endif // PLATFORM_WINDOWS
+	// Crank up a normal Slate application using the platform's standalone renderer.
+	FSlateApplication::InitializeAsStandaloneApplication(GetStandardStandaloneRenderer());
 
-#if PLATFORM_UNIX
-	int FileHandle = open("/var/run/UnrealInsightsBrowser.pid", O_CREAT | O_RDWR, 0666);
-	int Ret = flock(FileHandle, LOCK_EX | LOCK_NB);
-	if (Ret && EWOULDBLOCK == errno)
+	// Load required modules.
+	FModuleManager::Get().LoadModuleChecked("TraceInsightsCore");
+	if (bFrontendMode)
 	{
-		// Another Session Browser process is already running.
-
-		// Activate the respective window.
-		//TODO: "wmctrl -a Insights"
-
-		return false;
+		FModuleManager::Get().LoadModuleChecked("TraceInsightsFrontend");
 	}
-#endif
-
-	return true;
-}
-
-////////////////////////////////////////////////////////////////////////////////////////////////////
-
-void FUserInterfaceCommand::Run()
-{
-	const uint32 MaxPath = FPlatformMisc::GetMaxPathLength();
-	TCHAR* TraceFile = new TCHAR[MaxPath + 1];
-	TraceFile[0] = 0;
-	bool bOpenTraceFile = false;
-
-	// Only a single instance of Session Browser window/process is allowed.
+	else
 	{
-		bool bBrowserMode = true;
-
-		if (bBrowserMode)
-		{
-			bBrowserMode = FCString::Strifind(FCommandLine::Get(), TEXT("-OpenTraceId=")) == nullptr;
-		}
-		if (bBrowserMode)
-		{
-			bOpenTraceFile = GetTraceFileFromCmdLine(TraceFile, MaxPath);
-			bBrowserMode = !bOpenTraceFile;
-		}
-
-		if (bBrowserMode && !CheckSessionBrowserSingleInstance())
-		{
-			return;
-		}
+		FModuleManager::Get().LoadModuleChecked("TraceInsights");
 	}
 
-	//FCoreStyle::ResetToDefault();
-
-	// Crank up a normal Slate application using the platform's standalone renderer.
-	FSlateApplication::InitializeAsStandaloneApplication(GetStandardStandaloneRenderer());
-
-	// Load required modules.
-	FModuleManager::Get().LoadModuleChecked("TraceInsights");
-
 	// Load plug-ins.
-	// @todo: allow for better plug-in support in standalone Slate applications
 	IPluginManager::Get().LoadModulesForEnabledPlugins(ELoadingPhase::PreDefault);
 	IPluginManager::Get().LoadModulesForEnabledPlugins(ELoadingPhase::Default);
 
@@ -174,14 +96,13 @@ void FUserInterfaceCommand::Run()
 		FModuleManager::Get().LoadModule("SettingsEditor");
 	}
 
-	InitializeSlateApplication(bOpenTraceFile, TraceFile);
-
-	delete[] TraceFile;
-	TraceFile = nullptr;
+	InitializeSlateApplication(bFrontendMode, TraceFileToOpen);
 
 	IPluginManager::Get().LoadModulesForEnabledPlugins(ELoadingPhase::PostDefault);
 
+	//////////////////////////////////////////////////
 	// Initialize source code access.
+	
 	// Load the source code access module.
 	ISourceCodeAccessModule& SourceCodeAccessModule = FModuleManager::LoadModuleChecked<ISourceCodeAccessModule>(FName("SourceCodeAccess"));
 
@@ -194,6 +115,8 @@ void FUserInterfaceCommand::Run()
 	SourceCodeAccessModule.SetAccessor(FName("VisualStudioSourceCodeAccess"));
 #endif
 
+	//////////////////////////////////////////////////
+
 #if WITH_SHARED_POINTER_TESTS
 	SharedPointerTesting::TestSharedPointer<ESPMode::NotThreadSafe>();
 	SharedPointerTesting::TestSharedPointer<ESPMode::ThreadSafe>();
@@ -219,8 +142,7 @@ void FUserInterfaceCommand::Run()
 		FTSTicker::GetCoreTicker().Tick(static_cast<float>(DeltaTime));
 
 		// Throttle frame rate.
-		const float FrameTime = bDisableFramerateThrottle ? IdealFrameTime
-														  : UserInterfaceCommand::IsApplicationBackground() ? BackgroundFrameTime : IdealFrameTime;
+		const float FrameTime = !bDisableFramerateThrottle && UserInterfaceCommand::IsApplicationBackground() ? BackgroundFrameTime : IdealFrameTime;
 
 		UserInterfaceCommand::AdaptiveSleep(FMath::Max<float>(0.0f, FrameTime - static_cast<float>(FPlatformTime::Seconds() - LastTime)));
 
@@ -231,19 +153,17 @@ void FUserInterfaceCommand::Run()
 		FStats::AdvanceFrame(false);
 
 		FCoreDelegates::OnEndFrame.Broadcast();
-		GLog->FlushThreadedLogs(); //im: ???
+		GLog->FlushThreadedLogs();
 
 		GFrameCounter++;
 	}
 
-	//im: ??? FCoreDelegates::OnExit.Broadcast();
-
 	ShutdownSlateApplication();
 }
 
 ////////////////////////////////////////////////////////////////////////////////////////////////////
 
-void FUserInterfaceCommand::InitializeSlateApplication(bool bOpenTraceFile, const TCHAR* TraceFile)
+void FUserInterfaceCommand::InitializeSlateApplication(bool bFrontendMode, const FString& TraceFileToOpen)
 {
 	FSlateApplication::InitHighDPI(true);
 
@@ -261,21 +181,7 @@ void FUserInterfaceCommand::InitializeSlateApplication(bool bOpenTraceFile, cons
 		FModuleManager::LoadModuleChecked<ISlateReflectorModule>("SlateReflector").RegisterTabSpawner(UserInterfaceCommand::DeveloperTools);
 	}
 
-	IUnrealInsightsModule& TraceInsightsModule = FModuleManager::LoadModuleChecked<IUnrealInsightsModule>("TraceInsights");
-
-	uint32 TraceId = 0;
-	FString TraceIdString;
-	bool bUseTraceId = FParse::Value(FCommandLine::Get(), TEXT("-OpenTraceId="), TraceIdString);
-	if (TraceIdString.StartsWith(TEXT("0x")))
-	{
-		TCHAR* End;
-		TraceId = FCString::Strtoi(*TraceIdString + 2, &End, 16);
-	}
-	else
-	{
-		TCHAR* End;
-		TraceId = FCString::Strtoi(*TraceIdString, &End, 10);
-	}
+	//////////////////////////////////////////////////
 
 	FString StoreHost = TEXT("127.0.0.1");
 	uint32 StorePort = 0;
@@ -300,25 +206,42 @@ void FUserInterfaceCommand::InitializeSlateApplication(bool bOpenTraceFile, cons
 		bUseCustomStoreAddress = true;
 	}
 
-	// This parameter will cause the application to close when analysis fails to start or completes successfully.
-	const bool bAutoQuit = FParse::Param(FCommandLine::Get(), TEXT("AutoQuit"));
+	//////////////////////////////////////////////////
 
-	const bool bInitializeTesting = FParse::Param(FCommandLine::Get(), TEXT("InsightsTest"));
-	if (bInitializeTesting)
+	if (!bFrontendMode) // viewer mode
 	{
-		const bool bInitAutomationModules = true;
-		TraceInsightsModule.InitializeTesting(bInitAutomationModules, bAutoQuit);
-	}
+		FModuleManager::Get().LoadModuleChecked("TraceInsights");
+		IUnrealInsightsModule& TraceInsightsModule = FModuleManager::LoadModuleChecked<IUnrealInsightsModule>("TraceInsights");
 
-	if (bUseTraceId || bOpenTraceFile) // viewer mode
-	{
-		FString Cmd;
-		bool bExecuteCommand = false;
-		if (FParse::Value(FCommandLine::Get(), TEXT("-ExecOnAnalysisCompleteCmd="), Cmd, false))
+		// This parameter will cause the application to close when analysis fails to start or completes successfully.
+		const bool bAutoQuit = FParse::Param(FCommandLine::Get(), TEXT("AutoQuit"));
+
+		const bool bInitializeTesting = FParse::Param(FCommandLine::Get(), TEXT("InsightsTest"));
+		if (bInitializeTesting)
+		{
+			const bool bInitAutomationModules = true;
+			TraceInsightsModule.InitializeTesting(bInitAutomationModules, bAutoQuit);
+		}
+
+		uint32 TraceId = 0;
+		FString TraceIdString;
+		bool bUseTraceId = FParse::Value(FCommandLine::Get(), TEXT("-OpenTraceId="), TraceIdString);
+		if (bUseTraceId)
 		{
-			bExecuteCommand = true;
+			if (TraceIdString.StartsWith(TEXT("0x")))
+			{
+				TCHAR* End;
+				TraceId = FCString::Strtoi(*TraceIdString + 2, &End, 16);
+			}
+			else
+			{
+				TCHAR* End;
+				TraceId = FCString::Strtoi(*TraceIdString, &End, 10);
+			}
 		}
-		if (bExecuteCommand)
+
+		FString Cmd;
+		if (FParse::Value(FCommandLine::Get(), TEXT("-ExecOnAnalysisCompleteCmd="), Cmd, false))
 		{
 			TraceInsightsModule.ScheduleCommand(Cmd);
 		}
@@ -336,12 +259,15 @@ void FUserInterfaceCommand::InitializeSlateApplication(bool bOpenTraceFile, cons
 		}
 		else
 		{
-			TraceInsightsModule.StartAnalysisForTraceFile(TraceFile, bAutoQuit);
+			TraceInsightsModule.StartAnalysisForTraceFile(*TraceFileToOpen, bAutoQuit);
 		}
 	}
-	else // browser mode
+	else // frontend mode
 	{
-		// ensure target platform manager is referenced early as it must be created on the main thread
+		FModuleManager::Get().LoadModuleChecked("TraceInsightsFrontend");
+		ITraceInsightsFrontendModule& TraceInsightsFrontendModule = FModuleManager::LoadModuleChecked<ITraceInsightsFrontendModule>("TraceInsightsFrontend");
+
+		// Ensure target platform manager is referenced early as it must be created on the main thread.
 		FModuleManager::Get().LoadModuleChecked("DesktopPlatform");
 		FConfigCacheIni::InitializeConfigSystem();
 		GetTargetPlatformManager();
@@ -349,36 +275,27 @@ void FUserInterfaceCommand::InitializeSlateApplication(bool bOpenTraceFile, cons
 		FModuleManager::Get().LoadModuleChecked("Messaging");
 		FModuleManager::Get().LoadModuleChecked("OutputLog");
 
-		// load optional modules
+		// Load optional modules.
 		FModuleManager::Get().LoadModule("DeviceManager");
 		FModuleManager::Get().LoadModule("SessionFrontend");
 
-		FString Cmd;
-		bool bExecuteCommand = false;
-		if (FParse::Value(FCommandLine::Get(), TEXT("-ExecBrowserAutomationTest="), Cmd, false))
-		{
-			bExecuteCommand = true;
-		}
+		FString AutomationTests;
+		bool bRunAutomationTests =
+			FParse::Value(FCommandLine::Get(), TEXT("-ExecBrowserAutomationTest="), AutomationTests, false) ||
+			FParse::Value(FCommandLine::Get(), TEXT("-RunAutomationTests="), AutomationTests, false);
 
-		if (bUseCustomStoreAddress)
-		{
-			TraceInsightsModule.ConnectToStore(*StoreHost, StorePort);
-		}
-		else
-		{
-			TraceInsightsModule.CreateDefaultStore();
-		}
+		TraceInsightsFrontendModule.ConnectToStore(*StoreHost, StorePort);
 
-		FCreateSessionBrowserParams Params;
+		UE::Insights::FCreateFrontendWindowParams Params;
 		Params.bAllowDebugTools = bAllowDebugTools;
-		Params.bInitializeTesting = bInitializeTesting;
+		Params.bInitializeTesting = FParse::Param(FCommandLine::Get(), TEXT("InsightsTest"));
 		Params.bStartProcessWithStompMalloc = FParse::Param(FCommandLine::Get(), TEXT("stompmalloc"));
 		Params.bDisableFramerateThrottle = FParse::Param(FCommandLine::Get(), TEXT("DisableFramerateThrottle"));
-		TraceInsightsModule.CreateSessionBrowser(Params);
+		TraceInsightsFrontendModule.CreateFrontendWindow(Params);
 
-		if (bExecuteCommand)
+		if (bRunAutomationTests)
 		{
-			TraceInsightsModule.RunAutomationTest(Cmd);
+			TraceInsightsFrontendModule.RunAutomationTests(AutomationTests);
 		}
 	}
 }
@@ -395,33 +312,3 @@ void FUserInterfaceCommand::ShutdownSlateApplication()
 }
 
 ////////////////////////////////////////////////////////////////////////////////////////////////////
-
-bool FUserInterfaceCommand::GetTraceFileFromCmdLine(TCHAR* OutTraceFile, uint32 MaxPath)
-{
-	// Try getting the trace file from the -OpenTraceFile= paramter first.
-	bool bUseTraceFile = FParse::Value(FCommandLine::Get(), TEXT("-OpenTraceFile="), OutTraceFile, MaxPath, true);
-
-	if (bUseTraceFile)
-	{
-		return true;
-	}
-
-	// Support opening a trace file by double clicking a .utrace file.
-	// In this case, the app will receive as the first parameter a utrace file path.
-
-	const TCHAR* CmdLine = FCommandLine::Get();
-	bool HasToken = FParse::Token(CmdLine, OutTraceFile, MaxPath, false);
-
-	if (HasToken)
-	{
-		FString Token = OutTraceFile;
-		if (Token.EndsWith(TEXT(".utrace")))
-		{
-			bUseTraceFile = true;
-		}
-	}
-
-	return bUseTraceFile;
-}
-
-////////////////////////////////////////////////////////////////////////////////////////////////////
diff --git a/Engine/Source/Programs/UnrealInsights/Private/UserInterfaceCommand.h b/Engine/Source/Programs/UnrealInsights/Private/UserInterfaceCommand.h
--- a/Engine/Source/Programs/UnrealInsights/Private/UserInterfaceCommand.h
+++ b/Engine/Source/Programs/UnrealInsights/Private/UserInterfaceCommand.h
@@ -8,23 +8,16 @@ class FUserInterfaceCommand
 {
 public:
 	/** Executes the command. */
-	static void Run();
+	static void Run(bool bFrontendMode, const FString& TraceFileToOpen);
 
 protected:
 	/**
 	 * Initializes the Slate application.
 	 */
-	static void InitializeSlateApplication(bool bOpenTraceFile, const TCHAR* TraceFile);
+	static void InitializeSlateApplication(bool bFrontendMode, const FString& TraceFileToOpen);
 
 	/**
 	 * Shuts down the Slate application.
 	 */
 	static void ShutdownSlateApplication();
-
-private:
-	/**
-	* Attempts to get a utrace file path from the command line.
-	* Returns true if a path was found and false otherwise.
-	*/
-	static bool GetTraceFileFromCmdLine(TCHAR* OutTraceFile, uint32 MaxPath);
 };
diff --git a/Engine/Source/Programs/UnrealInsights/UnrealInsights.Build.cs b/Engine/Source/Programs/UnrealInsights/UnrealInsights.Build.cs
--- a/Engine/Source/Programs/UnrealInsights/UnrealInsights.Build.cs
+++ b/Engine/Source/Programs/UnrealInsights/UnrealInsights.Build.cs
@@ -24,6 +24,7 @@ public class UnrealInsights : ModuleRules
 				"TargetPlatform",
 				"TraceInsights",
 				"TraceInsightsCore",
+				"TraceInsightsFrontend",
 			}
 		);
 
