[Insights] Fixed compile error in UnrealInsights.rc file (re include Version.h).

#rnx

diff --git a/Engine/Source/Programs/UnrealInsights/Resources/Windows/UnrealInsights.rc b/Engine/Source/Programs/UnrealInsights/Resources/Windows/UnrealInsights.rc
--- a/Engine/Source/Programs/UnrealInsights/Resources/Windows/UnrealInsights.rc
+++ b/Engine/Source/Programs/UnrealInsights/Resources/Windows/UnrealInsights.rc
@@ -8,7 +8,7 @@
 #include <windows.h>
 #include "Runtime/Launch/Resources/Version.h"
 #include "Runtime/Launch/Resources/Windows/PCLaunch.rc"
-#include "Developer/TraceInsights/Public/Insights/Version.h"
+#include "Developer/TraceInsightsCore/Public/InsightsCore/Version.h"
 
 #define VER_FILEVERSION         UNREAL_INSIGHTS_VERSION_MAJOR,UNREAL_INSIGHTS_VERSION_MINOR,UNREAL_INSIGHTS_VERSION_PATCH,0
 #define VER_FILEVERSION_STR     UNREAL_INSIGHTS_VERSION_STRING "\0"
