Only enable LiveCoding on UnrealInsights for Win64 platforms

diff --git a/Engine/Source/Programs/UnrealInsights/UnrealInsights.Target.cs b/Engine/Source/Programs/UnrealInsights/UnrealInsights.Target.cs
--- a/Engine/Source/Programs/UnrealInsights/UnrealInsights.Target.cs
+++ b/Engine/Source/Programs/UnrealInsights/UnrealInsights.Target.cs
@@ -33,8 +33,11 @@ public class UnrealInsightsTarget : TargetRules
 
 		// For UI functionality
 		bBuildDeveloperTools = true;
-		
-		bWithLiveCoding = true;
+
+		if (Target.Platform == UnrealTargetPlatform.Win64 && Architecture.bIsX64 && Configuration != UnrealTargetConfiguration.Shipping)
+		{
+			bWithLiveCoding = true;
+		}
 
 		bHasExports = false;
 
