Wrap loading LiveCoding module in UnrealInsights with #ifdef
* This shouldn't matter but helps with finding LIVE_CODING support code.  UnrealInsights on non LiveCoding supported projects also won't have to try and fail to load a module.

diff --git a/Engine/Source/Programs/UnrealInsights/Private/UserInterfaceCommand.cpp b/Engine/Source/Programs/UnrealInsights/Private/UserInterfaceCommand.cpp
--- a/Engine/Source/Programs/UnrealInsights/Private/UserInterfaceCommand.cpp
+++ b/Engine/Source/Programs/UnrealInsights/Private/UserInterfaceCommand.cpp
@@ -76,9 +76,10 @@ void FUserInterfaceCommand::Run(bool bFrontendMode, const FString& TraceFileToOp
 	// Crank up a normal Slate application using the platform's standalone renderer.
 	FSlateApplication::InitializeAsStandaloneApplication(GetStandardStandaloneRenderer());
 
-	// Used by Live Coding, among other things.
+#if WITH_LIVE_CODING
 	FPlatformMisc::SetUBTTargetName(TEXT("UnrealInsights"));
 	FModuleManager::Get().LoadModule("LiveCoding");
+#endif
 
 	// Load required modules.
 	FModuleManager::Get().LoadModuleChecked("TraceInsightsCore");
