Disable LiveCoding for UnrealInsights by default
* LiveCoding starts up the module on an async thread by default.  This can be problematic for Insights if the application shutsdown quickly due to a race condition where the application is shutdown before insights has had a chance to start up.
* Disabling live coding with insights for now while more generalized fix can be found.

diff --git a/Engine/Source/Programs/UnrealInsights/UnrealInsights.Target.cs b/Engine/Source/Programs/UnrealInsights/UnrealInsights.Target.cs
--- a/Engine/Source/Programs/UnrealInsights/UnrealInsights.Target.cs
+++ b/Engine/Source/Programs/UnrealInsights/UnrealInsights.Target.cs
@@ -36,7 +36,9 @@ public class UnrealInsightsTarget : TargetRules
 
 		if (Target.Platform == UnrealTargetPlatform.Win64 && Architecture.bIsX64 && Configuration != UnrealTargetConfiguration.Shipping)
 		{
-			bWithLiveCoding = true;
+			// UE-274428 - Disabling by default.  LiveCoding starts up Async and can cause ensures when Insights
+			// is quickly shutdown
+			// bWithLiveCoding = true;
 		}
 
 		bHasExports = false;
