[Insights] Added build target settings to allow enabling tracing (-Trace, bEnableTrace and related settings) and/or debugging for the TraceAnalysis module (-TraceAnalysisDebug, bTraceAnalysisDebug and the related settings).

#rnx

diff --git a/Engine/Source/Programs/UnrealInsights/UnrealInsights.Target.cs b/Engine/Source/Programs/UnrealInsights/UnrealInsights.Target.cs
--- a/Engine/Source/Programs/UnrealInsights/UnrealInsights.Target.cs
+++ b/Engine/Source/Programs/UnrealInsights/UnrealInsights.Target.cs
@@ -10,6 +10,12 @@ public class UnrealInsightsTarget : TargetRules
 	[CommandLine("-Monolithic")]
 	public bool bMonolithic = false;
 
+	[CommandLine("-Trace")]
+	public bool bCmdLineTrace = false;
+
+	[CommandLine("-TraceAnalysisDebug")]
+	public bool bTraceAnalysisDebug = false;
+
 	public UnrealInsightsTarget(TargetInfo Target) : base(Target)
 	{
 		Type = TargetType.Program;
@@ -47,5 +53,42 @@ public class UnrealInsightsTarget : TargetRules
 		GlobalDefinitions.Add("UE_TRACE_SERVER_CONTROLS_ENABLED=1");
 		// Have UnrealInsights implicitly launch the trace store.
 		GlobalDefinitions.Add("UE_TRACE_SERVER_LAUNCH_ENABLED=1");
+
+		// UE_TRACE_ENABLED
+		bEnableTrace = bCmdLineTrace;
+		//bEnableTrace = true; // local override
+		if (bEnableTrace)
+		{
+			// Enable memory tracing
+			GlobalDefinitions.Add("UE_MEMORY_TRACE_AVAILABLE=1"); // allow memory tracing
+			GlobalDefinitions.Add("UE_MEMORY_TAGS_TRACE_ENABLED=1"); // enable scoped memory tags tracing
+			//GlobalDefinitions.Add("UE_CALLSTACK_TRACE_ENABLED=0"); // disable callstack tracing
+
+			// Enable LLM (Low Level Memory Tracker)
+			GlobalDefinitions.Add("LLM_ENABLED_IN_CONFIG=1");
+
+			// Disable various tracing systems
+			//GlobalDefinitions.Add("LOGTRACE_ENABLED=0");
+			//GlobalDefinitions.Add("MISCTRACE_ENABLED=0");
+			//GlobalDefinitions.Add("CPUPROFILERTRACE_ENABLED=0");
+			//GlobalDefinitions.Add("GPUPROFILERTRACE_ENABLED=0"); // old GPU Profiler
+			//GlobalDefinitions.Add("UE_TRACE_GPU_PROFILER_ENABLED=0"); // new GPU Profiler
+			//GlobalDefinitions.Add("LOADTIMEPROFILERTRACE_ENABLED=0");
+			//GlobalDefinitions.Add("STATSTRACE_ENABLED=0");
+			//GlobalDefinitions.Add("COUNTERSTRACE_ENABLED=0");
+			//GlobalDefinitions.Add("UE_TASK_TRACE_ENABLED=0");
+			//GlobalDefinitions.Add("UE_NET_TRACE_ENABLED=0");
+		}
+
+		// Debug for TraceAnalysis module.
+		//bTraceAnalysisDebug = true; // local override
+		if (bTraceAnalysisDebug)
+		{
+			// See Engine\Source\Developer\TraceAnalysis\Public\TraceAnalysisDebug.h
+			GlobalDefinitions.Add("UE_TRACE_ANALYSIS_DEBUG_API=1");
+			GlobalDefinitions.Add("UE_TRACE_ANALYSIS_DEBUG=1");
+			GlobalDefinitions.Add("UE_TRACE_ANALYSIS_DEBUG_LOG_IMPL=4");
+			GlobalDefinitions.Add("UE_TRACE_ANALYSIS_DEBUG_LEVEL=1");
+		}
 	}
 }
