Marking the main Unreal Insights module as requiring platform sdk

This module contains platform specific implementations, so it needs to be marked as surch for an upcoming feature (currently disabled by default)

#Synced-CL 50408994

diff --git a/Engine/Source/Programs/UnrealInsights/UnrealInsights.Build.cs b/Engine/Source/Programs/UnrealInsights/UnrealInsights.Build.cs
--- a/Engine/Source/Programs/UnrealInsights/UnrealInsights.Build.cs
+++ b/Engine/Source/Programs/UnrealInsights/UnrealInsights.Build.cs
@@ -8,6 +8,8 @@ public class UnrealInsights : ModuleRules
 	{
 		CppCompileWarningSettings.UnsafeTypeCastWarningLevel = WarningLevel.Error;
 
+		bRequiresPlatformSDK = true;
+
 		PublicIncludePathModuleNames.Add("Launch");
 
 		PrivateDependencyModuleNames.AddRange(
