[Insights] Add support in UnrealInsights to connect to secure trace.

Added functionality in Unreal Insights user interface commands to start a direct secure trace stream from command line parameters. Currently only supported for SelfSignedCertficiate and PreSharedKeys methods.

#Synced-CL 50345636

diff --git a/Engine/Source/Programs/UnrealInsights/Private/UnrealInsightsMain.cpp b/Engine/Source/Programs/UnrealInsights/Private/UnrealInsightsMain.cpp
--- a/Engine/Source/Programs/UnrealInsights/Private/UnrealInsightsMain.cpp
+++ b/Engine/Source/Programs/UnrealInsights/Private/UnrealInsightsMain.cpp
@@ -121,7 +121,8 @@ int32 UnrealInsightsMain(const TCHAR* CommandLine)
 	FString TraceFileToOpen;
 
 	if (FCString::Strifind(CommandLine, TEXT("-OpenTraceId=")) ||
-		FCString::Strifind(CommandLine, TEXT("-ListenForDirectTrace")))
+		FCString::Strifind(CommandLine, TEXT("-ListenForDirectTrace")) ||
+		FCString::Strifind(CommandLine, TEXT("-TraceSecurePort=")))
 	{
 		bFrontendMode = false;
 	}
diff --git a/Engine/Source/Programs/UnrealInsights/Private/UserInterfaceCommand.cpp b/Engine/Source/Programs/UnrealInsights/Private/UserInterfaceCommand.cpp
--- a/Engine/Source/Programs/UnrealInsights/Private/UserInterfaceCommand.cpp
+++ b/Engine/Source/Programs/UnrealInsights/Private/UserInterfaceCommand.cpp
@@ -3,7 +3,7 @@
 #include "UserInterfaceCommand.h"
 
 #include "Async/TaskGraphInterfaces.h"
-//#include "Brushes/SlateImageBrush.h"
+#include "Containers/AnsiString.h"
 #include "Containers/Ticker.h"
 #include "CoreGlobals.h"
 #include "Framework/Application/SlateApplication.h"
@@ -11,16 +11,17 @@
 #include "Framework/Docking/TabManager.h"
 #include "HAL/PlatformApplicationMisc.h"
 #include "HAL/PlatformProcess.h"
-#include "Interfaces/IPluginManager.h"
 #include "ISlateReflectorModule.h"
 #include "ISourceCodeAccessModule.h"
+#include "Interfaces/IPluginManager.h"
+#include "Misc/Base64.h"
 #include "Misc/CommandLine.h"
 #include "Misc/ConfigCacheIni.h"
 #include "Misc/Parse.h"
 #include "Modules/ModuleManager.h"
 #include "StandaloneRenderer.h"
-#include "Styling/AppStyle.h"
 #include "Stats/StatsSystem.h"
+#include "Styling/AppStyle.h"
 #include "Widgets/Docking/SDockTab.h"
 
 // TraceInsightsCore
@@ -42,6 +43,7 @@
 
 namespace UserInterfaceCommand
 {
+	const TCHAR* EngineExitInvalidArgsReason = TEXT("Invalid arguments");
 	TSharedRef<FWorkspaceItem> DeveloperTools = FWorkspaceItem::NewGroup(NSLOCTEXT("UnrealInsights", "DeveloperToolsMenu", "Developer Tools"));
 
 	bool IsApplicationBackground()
@@ -256,6 +258,9 @@ void FUserInterfaceCommand::InitializeSlateApplication(bool bFrontendMode, const
 			FParse::Value(FCommandLine::Get(), TEXT("-DirectTracePort="), DirectTracePort);
 		}
 
+		uint16 SecurePort = 0;
+		const bool bConnectToSecureTrace = FParse::Value(FCommandLine::Get(), TEXT("-TraceSecurePort="), SecurePort);
+
 		FString Cmd;
 		if (FParse::Value(FCommandLine::Get(), TEXT("-ExecOnAnalysisCompleteCmd="), Cmd, false))
 		{
@@ -269,7 +274,44 @@ void FUserInterfaceCommand::InitializeSlateApplication(bool bFrontendMode, const
 		}
 
 		TraceInsightsModule.SetAutoQuit(bAutoQuit, bWaitForSymbolResolver);
-		if (bListenForDirectTrace)
+		if (bConnectToSecureTrace)
+		{
+			FString SecureAuthBase64, SecureCertBase64, SecureHost(TEXT("localhost"));
+			FString SecureCertPath, SecureIdentity;
+			TArray<uint8> SecureAuthBin, SecureCertBin;
+
+			FParse::Value(FCommandLine::Get(), TEXT("-TraceSecureHost="), SecureHost);
+			FParse::Value(FCommandLine::Get(), TEXT("-TraceSecureAuth="), SecureAuthBase64);
+			FParse::Value(FCommandLine::Get(), TEXT("-TraceSecureCertPath="), SecureCertPath);
+			FParse::Value(FCommandLine::Get(), TEXT("-TraceSecureCert="), SecureCertBase64);
+			FParse::Value(FCommandLine::Get(), TEXT("-TraceSecureIdentity="), SecureIdentity);
+
+			FAnsiString SecureIdentityAnsi(SecureIdentity);
+
+			if (!SecureAuthBase64.IsEmpty() && !FBase64::Decode(SecureAuthBase64, SecureAuthBin))
+			{
+				UE_LOG(LogInit, Error, TEXT("Failed to decode secure authorization token."));
+				RequestEngineExit(UserInterfaceCommand::EngineExitInvalidArgsReason);
+				return;
+			}
+
+			if (!SecureCertBase64.IsEmpty() && !FBase64::Decode(SecureCertBase64, SecureCertBin))
+			{
+				UE_LOG(LogInit, Error, TEXT("Failed to decode secure certificate."));
+				RequestEngineExit(UserInterfaceCommand::EngineExitInvalidArgsReason);
+				return;
+			}
+			
+			TraceInsightsModule.StartAnalysisWithSecureTrace(
+				*SecureHost, 
+				SecurePort, 
+				SecureAuthBin, 
+				*SecureCertPath,
+				SecureCertBin,
+				SecureIdentityAnsi
+			);
+		}
+		else if (bListenForDirectTrace)
 		{
 			TraceInsightsModule.StartAnalysisWithDirectTrace(nullptr, DirectTracePort);
 		}
